ControlRig & animation mode: initial release of Direct Mesh Control
- experimental feature allowing to turn sections of a mesh into selectable controls which lets animators grab the mesh directly
- it takes advantage of the ability to define polygroup labels on triangles in order to define interaction zones
- these interaction zones are then used to generate skeletal meshes on the fly, as well as shape libraries that can be used in ControlRig
- todo: document the code and api

#virtualized
#Synced-CL 52542108

diff --git a/Engine/Plugins/Experimental/Animation/DirectMeshControl/DirectMeshControl.uplugin b/Engine/Plugins/Experimental/Animation/DirectMeshControl/DirectMeshControl.uplugin
new file mode 100644
--- /dev/null
+++ b/Engine/Plugins/Experimental/Animation/DirectMeshControl/DirectMeshControl.uplugin
@@ -0,0 +1,66 @@
+﻿{
+	"FileVersion": 3,
+	"Version": 1,
+	"VersionName": "1.0",
+	"FriendlyName": "Direct Mesh Control",
+	"Description": "Animate using click & drag and surface selection.",
+	"Category" : "Animation",
+	"CreatedBy" : "Epic Games, Inc.",
+	"CreatedByURL": "https://epicgames.com/",
+	"DocsURL": "",
+	"MarketplaceURL": "",
+	"EnabledByDefault" : false,
+	"CanContainContent": true,
+	"IsExperimentalVersion": true,
+	"IsBetaVersion": false,
+	"Installed": false,
+	"Plugins": [
+		{
+			"Name": "GeometryProcessing",
+			"Enabled": true
+		},
+		{
+			"Name": "MeshModelingToolset",
+			"Enabled": true,
+			"TargetAllowList": [ "Editor" ]
+		},
+		{
+			"Name": "ModelingToolsEditorMode",
+			"Enabled": true,
+			"TargetAllowList": [ "Editor" ]
+		},
+		{
+			"Name": "SkeletalMeshModelingTools",
+			"Enabled": true,
+			"TargetAllowList": [ "Editor" ]
+		},
+		{
+			"Name": "ControlRig",
+			"Enabled": true
+		},
+		{
+			"Name": "RigVM",
+			"Enabled": true
+		},
+		{
+			"Name": "ComputeFramework",
+			"Enabled": true
+		},
+		{
+			"Name": "DeformerGraph",
+			"Enabled": true
+		}
+	],
+	"Modules": [
+		{
+			"Name": "DirectMeshControl",
+			"Type": "Editor",
+			"LoadingPhase": "Default"
+		},
+		{
+			"Name": "DirectMeshControlRig",
+			"Type": "Editor",
+			"LoadingPhase": "Default"
+		}
+	]
+}
\ No newline at end of file
diff --git a/Engine/Plugins/Experimental/Animation/DirectMeshControl/Source/DirectMeshControl/DirectMeshControl.Build.cs b/Engine/Plugins/Experimental/Animation/DirectMeshControl/Source/DirectMeshControl/DirectMeshControl.Build.cs
new file mode 100644
--- /dev/null
+++ b/Engine/Plugins/Experimental/Animation/DirectMeshControl/Source/DirectMeshControl/DirectMeshControl.Build.cs
@@ -0,0 +1,47 @@
+﻿// Copyright Epic Games, Inc. All Rights Reserved.
+
+using UnrealBuildTool;
+
+public class DirectMeshControl : ModuleRules
+{
+	public DirectMeshControl(ReadOnlyTargetRules Target) : base(Target)
+	{
+		PCHUsage = ModuleRules.PCHUsageMode.UseExplicitOrSharedPCHs;
+		
+		PublicDependencyModuleNames.AddRange(new string[]
+		{
+			"ModelingComponents"
+		});
+		
+		PrivateDependencyModuleNames.AddRange(
+			new string[]
+			{
+				"ComputeFramework",
+				"Core",
+				"CoreUObject",
+				"DeveloperSettings",
+				"DynamicMesh",
+				"EditorInteractiveToolsFramework",
+				"EditorSubsystem",
+				"Engine",
+				"GeometryCore",
+				"GeometryFramework",
+				"InputCore",
+				"InteractiveToolsFramework",
+				"MeshConversion",
+				"MeshDescription",
+				"MeshModelingToolsEditorOnly",
+				"ModelingToolsEditorMode",
+				"OptimusCore",
+				"Persona",
+				"SkeletalMeshDescription",
+				"SkeletalMeshModelingTools",
+				"SkeletalMeshUtilitiesCommon",
+				"Slate",
+				"SlateCore",
+				"StaticMeshDescription",
+				"UnrealEd"
+			}
+		);
+	}
+}
diff --git a/Engine/Plugins/Experimental/Animation/DirectMeshControl/Source/DirectMeshControl/Private/DMCMeshGenerationSubsystem.cpp b/Engine/Plugins/Experimental/Animation/DirectMeshControl/Source/DirectMeshControl/Private/DMCMeshGenerationSubsystem.cpp
new file mode 100644
--- /dev/null
+++ b/Engine/Plugins/Experimental/Animation/DirectMeshControl/Source/DirectMeshControl/Private/DMCMeshGenerationSubsystem.cpp
@@ -0,0 +1,104 @@
+// Copyright Epic Games, Inc. All Rights Reserved.
+
+#include "DMCMeshGenerationSubsystem.h"
+
+#include "Editor.h"
+#include "Misc/SecureHash.h"
+#include "Engine/SkeletalMesh.h"
+
+#include UE_INLINE_GENERATED_CPP_BY_NAME(DMCMeshGenerationSubsystem)
+
+namespace UE::DMC
+{
+
+uint32 GetDDCHash(USkeletalMesh* InSkeletalMesh)
+{
+	FSHAHash SHAHash;
+	if (InSkeletalMesh)
+	{
+		const FString DDCKey = InSkeletalMesh->GetDerivedDataKey();
+		FSHA1 SHA1;
+		SHA1.UpdateWithString(*DDCKey, DDCKey.Len());
+		SHA1.Final();
+		SHA1.GetHash(SHAHash.Hash);
+	}
+	return HashCombine(GetTypeHash(InSkeletalMesh), GetTypeHash(SHAHash));
+}
+	
+}
+
+using namespace UE::DMC;
+using namespace UE::Geometry;
+
+void UDMCMeshGenerationManager::Shutdown()
+{
+	SubSkeletalMeshes.Reset();
+}
+
+const FGroupSubMeshes& UDMCMeshGenerationManager::GetSubMeshes(USkeletalMesh* InSkeletalMesh, const FDynamicMesh3* Mesh, const FName LayerName)
+{
+	if (!InSkeletalMesh || !Mesh)
+	{
+		static FGroupSubMeshes Dummy;
+		return Dummy;
+	}
+	
+	FSubSkeletalMeshData& Data = SubSkeletalMeshes.FindOrAdd({InSkeletalMesh, LayerName});
+
+	const uint32 Hash = GetDDCHash(InSkeletalMesh);
+	if (Data.Hash != Hash)
+	{
+		Data.Hash = Hash;
+		Data.GroupSubMeshes.Rebuild(InSkeletalMesh, Mesh, LayerName);
+	}
+	else
+	{
+		// check consistency
+	}
+
+	return Data.GroupSubMeshes;
+}
+
+bool UDMCMeshGenerationSubsystem::bIsShuttingDown = false;
+
+UDMCMeshGenerationSubsystem* UDMCMeshGenerationSubsystem::Get()
+{
+	if (GEditor && !bIsShuttingDown)
+	{
+		return GEditor->GetEditorSubsystem<UDMCMeshGenerationSubsystem>();
+	}
+	return nullptr;
+}
+
+void UDMCMeshGenerationSubsystem::Initialize(FSubsystemCollectionBase& Collection)
+{
+	bIsShuttingDown = false;
+	GenerationManager = NewObject<UDMCMeshGenerationManager>(this);
+
+	if (GEditor)
+	{
+		GEditor->OnEditorClose().AddUObject(this, &UDMCMeshGenerationSubsystem::OnShutdown);
+	}
+	FCoreDelegates::OnEnginePreExit.AddUObject(this, &UDMCMeshGenerationSubsystem::OnShutdown);
+}
+
+void UDMCMeshGenerationSubsystem::Deinitialize()
+{
+	if (GEditor)
+	{
+		GEditor->OnEditorClose().RemoveAll(this);
+	}
+	FCoreDelegates::OnEnginePreExit.RemoveAll(this);
+	bIsShuttingDown = true;
+
+	if (GenerationManager)
+	{
+		GenerationManager->Shutdown();
+		GenerationManager = nullptr;
+	}
+}
+
+void UDMCMeshGenerationSubsystem::OnShutdown()
+{
+	bIsShuttingDown = true;
+}
diff --git a/Engine/Plugins/Experimental/Animation/DirectMeshControl/Source/DirectMeshControl/Private/DMCMeshGenerationSubsystem.h b/Engine/Plugins/Experimental/Animation/DirectMeshControl/Source/DirectMeshControl/Private/DMCMeshGenerationSubsystem.h
new file mode 100644
--- /dev/null
+++ b/Engine/Plugins/Experimental/Animation/DirectMeshControl/Source/DirectMeshControl/Private/DMCMeshGenerationSubsystem.h
@@ -0,0 +1,65 @@
+// Copyright Epic Games, Inc. All Rights Reserved.
+
+#pragma once
+
+#include "EditorSubsystem.h"
+#include "DirectMeshControlUtilities.h"
+
+#include "DMCMeshGenerationSubsystem.generated.h"
+
+#define UE_API DIRECTMESHCONTROL_API
+
+class USkeletalMesh;
+class UMaterialInstanceDynamic;
+class UOptimusDeformer;
+
+namespace UE::Geometry
+{
+	class FDynamicMesh3;
+}
+
+/**
+ *
+ */
+
+UCLASS(MinimalAPI)
+class UDMCMeshGenerationManager : public UObject
+{
+	GENERATED_BODY()
+
+public:
+	UE_API virtual void Shutdown();
+
+	UE_API const FGroupSubMeshes& GetSubMeshes(USkeletalMesh* InSkeletalMesh, const UE::Geometry::FDynamicMesh3* Mesh, const FName LayerName);
+	
+protected:
+	
+	UPROPERTY()
+	TSet<FSubSkeletalMeshData> SubSkeletalMeshes;
+};
+	
+/**
+ *
+ */
+
+UCLASS(MinimalAPI)
+class UDMCMeshGenerationSubsystem : public UEditorSubsystem
+{
+	GENERATED_BODY()
+
+public:
+	static UDMCMeshGenerationSubsystem* Get();
+	UE_API virtual void Initialize(FSubsystemCollectionBase& Collection) override;
+	UE_API virtual void Deinitialize() override;
+
+	UPROPERTY()
+	TObjectPtr<UDMCMeshGenerationManager> GenerationManager = nullptr;
+	
+protected:
+	UE_API virtual void OnShutdown();
+
+private:
+	static UE_API bool bIsShuttingDown;
+};
+	
+#undef UE_API
diff --git a/Engine/Plugins/Experimental/Animation/DirectMeshControl/Source/DirectMeshControl/Private/DirectMeshControlCommands.cpp b/Engine/Plugins/Experimental/Animation/DirectMeshControl/Source/DirectMeshControl/Private/DirectMeshControlCommands.cpp
new file mode 100644
--- /dev/null
+++ b/Engine/Plugins/Experimental/Animation/DirectMeshControl/Source/DirectMeshControl/Private/DirectMeshControlCommands.cpp
@@ -0,0 +1,25 @@
+﻿// Copyright Epic Games, Inc. All Rights Reserved.
+
+#include "DirectMeshControlCommands.h"
+#include "InputCoreTypes.h"
+#include "DirectMeshControlStyle.h"
+
+#define LOCTEXT_NAMESPACE "DirectMeshControlCommands"
+
+FDirectMeshControlCommands::FDirectMeshControlCommands() :
+	TCommands<FDirectMeshControlCommands>(
+		"DirectMeshControlCommands", // Context name for fast lookup
+		NSLOCTEXT("Contexts", "DirectMeshControlCommands", "Direct Mesh Control Extension"), // Localized context name for displaying
+		NAME_None, // Parent
+		FDirectMeshControlStyle::Get()->GetStyleSetName() // Icon Style Set
+		)
+{
+}
+
+void FDirectMeshControlCommands::RegisterCommands()
+{
+	UI_COMMAND(BeginDirectMeshPolygroupTool, "DMC Polygroups", "Generate polygroups for surface selection", EUserInterfaceActionType::ToggleButton, FInputChord());
+	UI_COMMAND(BeginDirectMeshControlTool, "DMC Control", "Surface selection manipulation", EUserInterfaceActionType::ToggleButton, FInputChord());
+}
+
+#undef LOCTEXT_NAMESPACE
\ No newline at end of file
diff --git a/Engine/Plugins/Experimental/Animation/DirectMeshControl/Source/DirectMeshControl/Private/DirectMeshControlComponent.cpp b/Engine/Plugins/Experimental/Animation/DirectMeshControl/Source/DirectMeshControl/Private/DirectMeshControlComponent.cpp
new file mode 100644
--- /dev/null
+++ b/Engine/Plugins/Experimental/Animation/DirectMeshControl/Source/DirectMeshControl/Private/DirectMeshControlComponent.cpp
@@ -0,0 +1,87 @@
+// Copyright Epic Games, Inc. All Rights Reserved.
+
+
+#include "DirectMeshControlComponent.h"
+
+#include "SkeletalMeshDeformerHelpers.h"
+#include "SkeletalRenderPublic.h"
+#include "Animation/MeshDeformerInstance.h"
+#include "Engine/SkeletalMesh.h"
+#include "Rendering/SkeletalMeshRenderData.h"
+
+void UDirectMeshControlComponent::CreateRenderState_Concurrent(FRegisterComponentContext* Context)
+{
+	Super::CreateRenderState_Concurrent(Context);
+}
+
+void UDirectMeshControlComponent::SendRenderDynamicData_Concurrent()
+{
+#if WITH_STATE_STREAM_ACTOR
+	return;
+#endif
+
+	USkeletalMesh* SkeletalMeshPtr = GetSkeletalMeshAsset();
+	
+#if WITH_EDITOR
+	if (SkeletalMeshPtr && SkeletalMeshPtr->IsCompiling())
+	{
+		return;
+	}
+#endif
+	
+	// skinned mesh component override
+	// USkinnedMeshComponent::SendRenderDynamicData_Concurrent();
+
+	if (MeshObject && SkeletalMeshPtr && (bRecentlyRendered || VisibilityBasedAnimTickOption == EVisibilityBasedAnimTickOption::AlwaysTickPoseAndRefreshBones || GIsEditor || !MeshObject->bHasBeenUpdatedAtLeastOnce))
+	{
+		const int32 MinLodIndex = ComputeMinLOD();
+		const int32 MaxLODIndex = MeshObject->GetSkeletalMeshRenderData().LODRenderData.Num() - 1;
+		int32 UseLOD = FMath::Clamp(GetPredictedLODLevel(), MinLodIndex, MaxLODIndex);
+		
+		// Only update the state if PredictedLODLevel is valid
+		FSkeletalMeshRenderData* SkelMeshRenderData = GetSkeletalMeshRenderData();
+		if (SkelMeshRenderData && SkelMeshRenderData->LODRenderData.IsValidIndex(UseLOD))
+		{
+			ActiveMorphTargets.Empty();
+			MeshObject->Update(
+				UseLOD,
+				this,
+				ActiveMorphTargets,
+				MorphTargetWeights,
+				EPreviousBoneTransformUpdateMode::None,
+				GetExternalMorphWeights(UseLOD)
+			);  // send to rendering thread
+
+			MeshObject->bHasBeenUpdatedAtLeastOnce = true;
+		}
+
+		if (UMeshDeformerInstance* DeformerInstanceForLOD = GetMeshDeformerInstanceForLOD(UseLOD))
+		{
+#if WITH_EDITORONLY_DATA
+			for (TUniquePtr<FMeshDeformerGeometryReadbackRequest>& Request : GeometryReadbackRequests)
+			{
+				DeformerInstanceForLOD->RequestReadbackDeformerGeometry(MoveTemp(Request));
+			}
+#endif // WITH_EDITORONLY_DATA
+			
+			UMeshDeformerInstance::FEnqueueWorkDesc Desc;
+			Desc.Scene = GetScene();
+			Desc.OwnerName = SkeletalMeshPtr ? SkeletalMeshPtr->GetFName() : GetFName();
+			Desc.ExecutionGroup = UMeshDeformerInstance::ExecutionGroup_BeginInitViews;
+		
+			// Fallback is to reset the passthrough vertex factory if the deformer fails to run.
+			Desc.FallbackDelegate.BindLambda([MeshObjectPtr = MeshObject, UseLOD]()
+			{
+				FSkeletalMeshDeformerHelpers::ResetVertexFactoryBufferOverrides(MeshObjectPtr, UseLOD);
+			});
+			DeformerInstanceForLOD->EnqueueWork(Desc);
+		}
+	}
+
+	UActorComponent::SendRenderDynamicData_Concurrent();
+
+#if WITH_EDITORONLY_DATA
+	// Avoid unbound accumulation of readback requests, each readback request must be picked up by the deformer instance within the same frame or ignored otherwise
+	GeometryReadbackRequests.Reset();
+#endif // WITH_EDITORONLY_DATA 
+}
\ No newline at end of file
diff --git a/Engine/Plugins/Experimental/Animation/DirectMeshControl/Source/DirectMeshControl/Private/DirectMeshControlModule.cpp b/Engine/Plugins/Experimental/Animation/DirectMeshControl/Source/DirectMeshControl/Private/DirectMeshControlModule.cpp
new file mode 100644
--- /dev/null
+++ b/Engine/Plugins/Experimental/Animation/DirectMeshControl/Source/DirectMeshControl/Private/DirectMeshControlModule.cpp
@@ -0,0 +1,67 @@
+﻿// Copyright Epic Games, Inc. All Rights Reserved.
+
+#include "DirectMeshControlModule.h"
+
+#include "DirectMeshControlCommands.h"
+#include "DirectMeshControlStyle.h"
+#include "ModelingModeToolExtensions.h"
+#include "Features/IModularFeatures.h"
+#include "Tools/DirectMeshControlTool.h"
+#include "Tools/DirectMeshPolygroupTool.h"
+
+#include "InteractiveToolManager.h"
+#include "ToolMenus.h"
+#include "ViewportToolbar/UnrealEdViewportToolbarContext.h"
+#include "Tools/EdModeInteractiveToolsContext.h"
+
+#define LOCTEXT_NAMESPACE "FDirectMeshControlModule"
+
+void FDirectMeshControlModule::StartupModule()
+{
+	FDirectMeshControlStyle::Initialize();
+	FDirectMeshControlCommands::Register();
+
+	IModularFeatures::Get().RegisterModularFeature(GetModularFeatureName(), this);
+}
+
+void FDirectMeshControlModule::ShutdownModule()
+{
+	IModularFeatures::Get().UnregisterModularFeature(GetModularFeatureName(), this);
+
+	FDirectMeshControlCommands::Unregister();
+	FDirectMeshControlStyle::Shutdown();
+}
+
+FText FDirectMeshControlModule::GetExtensionName()
+{
+	return LOCTEXT("ExtensionName", "DMCTools");
+}
+
+FText FDirectMeshControlModule::GetToolSectionName()
+{
+	return LOCTEXT("SectionName", "DMC");
+}
+
+void FDirectMeshControlModule::GetExtensionTools(const FExtensionToolQueryInfo& InQueryInfo, TArray<FExtensionToolDescription>& OutTools)
+{
+	{
+		FExtensionToolDescription DirectMeshPolygroupToolInfo;
+		DirectMeshPolygroupToolInfo.ToolName = LOCTEXT("DirectMeshPolygroupTool", "Direct Mesh Polygroup");
+		DirectMeshPolygroupToolInfo.ToolCommand = FDirectMeshControlCommands::Get().BeginDirectMeshPolygroupTool;
+		DirectMeshPolygroupToolInfo.ToolBuilder = NewObject<UDirectMeshPolygroupToolBuilder>();
+		OutTools.Add(DirectMeshPolygroupToolInfo);
+	}
+	
+	{
+		FExtensionToolDescription DirectMeshControlToolInfo;
+		DirectMeshControlToolInfo.ToolName = LOCTEXT("DirectMeshControlTool", "Direct Mesh Control");
+		DirectMeshControlToolInfo.ToolCommand = FDirectMeshControlCommands::Get().BeginDirectMeshControlTool;
+		UDirectMeshControlToolBuilder* Builder = NewObject<UDirectMeshControlToolBuilder>();
+		DirectMeshControlToolInfo.ToolBuilder = Builder;
+		OutTools.Add(DirectMeshControlToolInfo);
+	}
+}
+
+#undef LOCTEXT_NAMESPACE
+	
+IMPLEMENT_MODULE(FDirectMeshControlModule, DirectMeshControl)
\ No newline at end of file
diff --git a/Engine/Plugins/Experimental/Animation/DirectMeshControl/Source/DirectMeshControl/Private/DirectMeshControlSettings.cpp b/Engine/Plugins/Experimental/Animation/DirectMeshControl/Source/DirectMeshControl/Private/DirectMeshControlSettings.cpp
new file mode 100644
--- /dev/null
+++ b/Engine/Plugins/Experimental/Animation/DirectMeshControl/Source/DirectMeshControl/Private/DirectMeshControlSettings.cpp
@@ -0,0 +1,17 @@
+// Copyright Epic Games, Inc. All Rights Reserved.
+
+#include "DirectMeshControlSettings.h"
+
+#include UE_INLINE_GENERATED_CPP_BY_NAME(DirectMeshControlSettings)
+
+UDirectMeshControlSettings::UDirectMeshControlSettings(const FObjectInitializer& ObjectInitializer)
+	: UDeveloperSettings(ObjectInitializer)
+{}
+
+UDirectMeshControlSettings::FOnUpdateSettings UDirectMeshControlSettings::OnSettingsChange;
+
+void UDirectMeshControlSettings::PostEditChangeProperty(FPropertyChangedEvent& PropertyChangedEvent)
+{
+	Super::PostEditChangeProperty(PropertyChangedEvent);
+	OnSettingsChange.Broadcast(this);
+}
diff --git a/Engine/Plugins/Experimental/Animation/DirectMeshControl/Source/DirectMeshControl/Private/DirectMeshControlStyle.cpp b/Engine/Plugins/Experimental/Animation/DirectMeshControl/Source/DirectMeshControl/Private/DirectMeshControlStyle.cpp
new file mode 100644
--- /dev/null
+++ b/Engine/Plugins/Experimental/Animation/DirectMeshControl/Source/DirectMeshControl/Private/DirectMeshControlStyle.cpp
@@ -0,0 +1,49 @@
+﻿// Copyright Epic Games, Inc. All Rights Reserved.
+
+#include "DirectMeshControlStyle.h"
+#include "Styling/SlateStyleRegistry.h"
+#include "Styling/SlateTypes.h"
+#include "Styling/CoreStyle.h"
+#include "SlateOptMacros.h"
+
+TSharedPtr< FSlateStyleSet > FDirectMeshControlStyle::StyleSet = nullptr;
+
+TSharedPtr< class ISlateStyle > FDirectMeshControlStyle::Get()
+{
+	return StyleSet;
+}
+
+FName FDirectMeshControlStyle::GetStyleSetName()
+{
+	static FName StyleName(TEXT("DirectMeshControlStyle"));
+	return StyleName;
+}
+
+BEGIN_SLATE_FUNCTION_BUILD_OPTIMIZATION
+
+void FDirectMeshControlStyle::Initialize()
+{
+	// Only register once
+	if (StyleSet.IsValid())
+	{
+		return;
+	}
+	
+	StyleSet = MakeShareable(new FSlateStyleSet(GetStyleSetName()));
+	StyleSet->SetContentRoot(FPaths::EngineContentDir() / TEXT("Editor/Slate"));
+	StyleSet->SetCoreContentRoot(FPaths::EngineContentDir() / TEXT("Slate"));
+
+	FSlateStyleRegistry::RegisterSlateStyle(*StyleSet.Get());
+};
+
+END_SLATE_FUNCTION_BUILD_OPTIMIZATION
+
+void FDirectMeshControlStyle::Shutdown()
+{
+	if (StyleSet.IsValid())
+	{
+		FSlateStyleRegistry::UnRegisterSlateStyle(*StyleSet.Get());
+		ensure(StyleSet.IsUnique());
+		StyleSet.Reset();
+	}
+}
\ No newline at end of file
diff --git a/Engine/Plugins/Experimental/Animation/DirectMeshControl/Source/DirectMeshControl/Private/DirectMeshControlUtilities.cpp b/Engine/Plugins/Experimental/Animation/DirectMeshControl/Source/DirectMeshControl/Private/DirectMeshControlUtilities.cpp
new file mode 100644
--- /dev/null
+++ b/Engine/Plugins/Experimental/Animation/DirectMeshControl/Source/DirectMeshControl/Private/DirectMeshControlUtilities.cpp
@@ -0,0 +1,284 @@
+// Copyright Epic Games, Inc. All Rights Reserved.
+
+#include "DirectMeshControlUtilities.h"
+
+#include "Animation/Skeleton.h"
+#include "DMCMeshGenerationSubsystem.h"
+#include "DynamicMeshToMeshDescription.h"
+#include "DynamicSubmesh3.h"
+#include "Engine/SkinnedAssetCommon.h"
+#include "Materials/Material.h"
+#include "MeshConversionOptions.h"
+#include "MeshDescription.h"
+#include "OptimusDeformer.h"
+#include "Polygroups/GroupSetAdapter.h"
+#include "Polygroups/PolygroupUtil.h"
+#include "SkeletalMeshAttributes.h"
+#include "Rendering/SkeletalMeshModel.h"
+#include "StaticToSkeletalMeshConverter.h"
+#include "UObject/Package.h"
+#include "Util/ColorConstants.h"
+
+#include UE_INLINE_GENERATED_CPP_BY_NAME(DirectMeshControlUtilities)
+
+namespace UE::DMC
+{
+	
+static bool bMaterialMode = true;
+	
+bool UseMaterialHover()
+{
+	return bMaterialMode;
+}
+	
+const FName& GetSubToSourceAttrName()
+{
+	static FName VtxMappingAttrName(TEXT("SubToSource"));
+	return VtxMappingAttrName;
+}
+
+const FName& GetColorVarName()
+{
+	static FName VtxColorAttrName(TEXT("OverlayColor"));
+	return VtxColorAttrName;
+}
+
+UMaterial* GetMaterial()
+{
+	if (bMaterialMode)
+	{
+		static TWeakObjectPtr<UMaterial> HoveredMaterial = nullptr;
+		if (!HoveredMaterial.IsValid())
+		{
+			static FString HoveredMaterialPath(TEXT("/DirectMeshControl/Materials/M_DirectMeshControlHover"));
+			HoveredMaterial = LoadObject<UMaterial>(nullptr, HoveredMaterialPath);
+			ensure(HoveredMaterial.IsValid());
+		}
+		return HoveredMaterial.Get();
+	}
+	
+	static TWeakObjectPtr<UMaterial> Material = nullptr;
+	if (!Material.IsValid())
+	{
+		static FString MaterialPath(TEXT("/DirectMeshControl/Materials/M_DirectMeshControl"));
+		Material = LoadObject<UMaterial>(nullptr, MaterialPath);
+		ensure(Material.IsValid());
+	}
+	return Material.Get();
+}
+
+USkeleton* GetSkeleton()
+{
+	static TWeakObjectPtr<USkeleton> WeakSkeleton = nullptr;
+	if (!WeakSkeleton.IsValid())
+	{
+		UDMCMeshGenerationSubsystem* Subsystem = UDMCMeshGenerationSubsystem::Get();
+		if (ensure(Subsystem))
+		{
+			WeakSkeleton = NewObject<USkeleton>(Subsystem);
+			{
+				const TCHAR* RootBoneName = TEXT("Root"); 
+				FTransform RootTransform(FTransform::Identity);
+				FReferenceSkeletonModifier Modifier(WeakSkeleton.Get());
+				Modifier.Add(FMeshBoneInfo(RootBoneName, RootBoneName, INDEX_NONE), RootTransform);
+			}
+		}
+	}
+	return WeakSkeleton.Get();
+}
+
+UOptimusDeformer* GetDeformer()
+{
+	if (bMaterialMode)
+	{
+		static TWeakObjectPtr<UOptimusDeformer> DeformerNoColor = nullptr;
+		if (!DeformerNoColor.IsValid())
+		{
+			static FString DeformerPath(TEXT("/DirectMeshControl/Deformers/DG_DirectMeshControlNoColor"));
+			DeformerNoColor = LoadObject<UOptimusDeformer>(nullptr, DeformerPath);
+			ensure(DeformerNoColor.IsValid());
+		}
+		return DeformerNoColor.Get();
+	}
+	
+	static TWeakObjectPtr<UOptimusDeformer> Deformer = nullptr;
+	if (!Deformer.IsValid())
+	{
+		static FString DeformerPath(TEXT("/DirectMeshControl/Deformers/DG_DirectMeshControl"));
+		Deformer = LoadObject<UOptimusDeformer>(nullptr, DeformerPath);
+		ensure(Deformer.IsValid());
+	}
+	return Deformer.Get();
+}
+	
+const FGroupSubMeshes& GetSubMeshes(USkeletalMesh* InSkeletalMesh, const Geometry::FDynamicMesh3* Mesh, const FName LayerName)
+{
+	static const FGroupSubMeshes Dummy;
+	if (!InSkeletalMesh || !Mesh)
+	{
+		return Dummy;
+	}
+
+	UDMCMeshGenerationSubsystem* Subsystem = UDMCMeshGenerationSubsystem::Get();
+	if (!ensure(Subsystem))
+	{
+		return Dummy;
+	}
+
+	return Subsystem->GenerationManager ? Subsystem->GenerationManager->GetSubMeshes(InSkeletalMesh, Mesh, LayerName) : Dummy;
+}
+	
+}
+
+using namespace UE::DMC;
+using namespace UE::Geometry;
+
+void FGroupSubMeshes::Rebuild(USkeletalMesh* InSkeletalMesh, const FDynamicMesh3* Mesh, const FName LayerName)
+{
+	Reset();
+
+	UDMCMeshGenerationSubsystem* Subsystem = UDMCMeshGenerationSubsystem::Get();
+	if (!ensure(Subsystem))
+	{
+		return;
+	}
+
+	const FDynamicMeshTriangleLabelAttribute* SelectionLayer = FindTriangleLabelLayerByName(*Mesh, LayerName);
+	if (!ensure(SelectionLayer))
+	{
+		return;
+	}
+
+	const FTriangleReadOnlyLabelAdapter Adapter(Mesh, SelectionLayer);
+	
+	TMap<int32, TArray<int32>> PolyGroupToTri;
+	for (int32 TriIndex = 0; TriIndex < Mesh->TriangleCount(); TriIndex++)
+	{
+		const int32 GroupID = Adapter.GetGroup(TriIndex);
+		PolyGroupToTri.FindOrAdd(GroupID).Add(TriIndex);
+	}
+
+	if (PolyGroupToTri.IsEmpty())
+	{
+		return;
+	}
+	
+	FSkeletalMeshModel* SkeletalMeshModel = InSkeletalMesh->GetImportedModel();
+	if (!SkeletalMeshModel || SkeletalMeshModel->LODModels.IsEmpty())
+	{
+		return;
+	}
+
+	USkeleton* Skeleton = GetSkeleton();
+	UMaterial* Material = GetMaterial();
+	if (!Skeleton || !Material)
+	{
+		return;
+	}
+	
+	const TArray<uint32>& RenderToImportedMap = SkeletalMeshModel->LODModels[0].GetRawPointIndices();
+	const TArray<FSkeletalMaterial> Materials({Material});
+	const TArray<FName> AttributesToEnableForRender({GetSubToSourceAttrName()});
+	
+	FStaticToSkeletalMeshConverter::FInitializationParams Parameters;
+	Parameters.Materials = Materials;
+	Parameters.bRecomputeNormals = false;
+	Parameters.bRecomputeTangents = false;
+	Parameters.bCacheOptimize = true;
+	Parameters.VertexAttributesForRender = AttributesToEnableForRender;
+	
+	for (const auto& [GroudID, Triangles]: PolyGroupToTri)
+	{
+		FMeshDescription MeshDescription;
+		{
+			FDynamicSubmesh3 SubMesh(Mesh, Triangles);
+			if (FDynamicMeshAttributeSet* SubAttributes = SubMesh.GetSubmesh().Attributes())
+			{
+				SubAttributes->DisableMaterialID();
+			}
+					
+			// Full binding to the root bone.
+			FSkeletalMeshAttributes Attributes(MeshDescription);
+			Attributes.Register();
+			FConversionToMeshDescriptionOptions ConverterOptions;
+			ConverterOptions.bSetPolyGroups = false;
+			ConverterOptions.MaxValidMaterialID = 0;
+			FDynamicMeshToMeshDescription Converter(ConverterOptions);
+			Converter.Convert(&SubMesh.GetSubmesh(), MeshDescription);
+
+			// Full binding to the root bone.
+			constexpr int32 RootBoneIndex = 0;
+			FSkinWeightsVertexAttributesRef SkinWeights = Attributes.GetVertexSkinWeights();
+			UE::AnimationCore::FBoneWeight RootInfluence(RootBoneIndex, 1.0f);
+			UE::AnimationCore::FBoneWeights RootBinding = UE::AnimationCore::FBoneWeights::Create({ RootInfluence });
+			for (const FVertexID VertexID : MeshDescription.Vertices().GetElementIDs())
+			{
+				SkinWeights.Set(VertexID, RootBinding);
+			}
+
+			// vertex mapping attributes
+			MeshDescription.VertexAttributes().RegisterAttribute<float>(GetSubToSourceAttrName(), 1, 0.5f, EMeshAttributeFlags::None);
+			TVertexAttributesRef<float> VtxMappingAttr = MeshDescription.VertexAttributes().GetAttributesRef<float>(GetSubToSourceAttrName());
+
+			for (const FVertexID VertexID : MeshDescription.Vertices().GetElementIDs())
+			{
+				const int32 BaseIndex = SubMesh.MapVertexToBaseMesh(VertexID);
+				const int32 RenderIndex = RenderToImportedMap.Find(BaseIndex);
+				if (ensure(RenderToImportedMap.IsValidIndex(RenderIndex)))
+				{
+					VtxMappingAttr.Set(VertexID, static_cast<float>(RenderIndex) + 0.5f);
+				}
+			}
+
+			// vertex colors
+			const FVector4f Color(LinearColors::SelectFColor(GroudID));
+			TVertexInstanceAttributesRef<FVector4f> VertexColor = Attributes.GetVertexInstanceColors();
+			for (int32 ColorIndex = 0; ColorIndex < VertexColor.GetNumElements(); ColorIndex++)
+			{
+				VertexColor.Set(ColorIndex, Color);
+			}
+		}
+
+		const FString Name = FString::Printf(TEXT("%s_%s_%d"), *InSkeletalMesh->GetName(), *LayerName.ToString(), GroudID);
+		const FName UniqueName = MakeUniqueObjectName(Subsystem->GenerationManager, USkeletalMesh::StaticClass(), *Name);
+		
+		USkeletalMesh* SubSkeletalMesh = NewObject<USkeletalMesh>(Subsystem->GenerationManager, UniqueName, RF_NoFlags);
+		{
+			TArray<const FMeshDescription*> MeshDescriptions({&MeshDescription});
+			FStaticToSkeletalMeshConverter::InitializeSkeletalMeshFromMeshDescriptions(
+				SubSkeletalMesh,
+				MeshDescriptions,
+				Skeleton->GetReferenceSkeleton(),
+				Parameters);
+		}
+		SubSkeletalMesh->SetSkeleton(Skeleton);
+		
+		const int32 Index = SubSkeletalMeshes.Add(SubSkeletalMesh);
+		GroupIdToSubMesh.Emplace(GroudID, Index);
+	}
+}
+
+void FGroupSubMeshes::Reset()
+{
+	for (TObjectPtr<USkeletalMesh>& WeakSubMesh: SubSkeletalMeshes)
+	{
+		if (WeakSubMesh.Get())
+		{
+			WeakSubMesh->Rename(nullptr, GetTransientPackage(), REN_DoNotDirty | REN_DontCreateRedirectors | REN_NonTransactional);
+			WeakSubMesh->MarkAsGarbage();
+			WeakSubMesh = nullptr;
+		}
+	}
+	SubSkeletalMeshes.Reset();
+	GroupIdToSubMesh.Reset();
+}
+
+const TArray<TObjectPtr<USkeletalMesh>>& FGroupSubMeshes::GetSubSkeletalMeshes() const
+{
+	return SubSkeletalMeshes;
+}
+	
+const TMap<int32, int32>& FGroupSubMeshes::GetGroupIdToSubMesh() const
+{
+	return GroupIdToSubMesh;
+}
diff --git a/Engine/Plugins/Experimental/Animation/DirectMeshControl/Source/DirectMeshControl/Private/Tools/DirectMeshControlTool.cpp b/Engine/Plugins/Experimental/Animation/DirectMeshControl/Source/DirectMeshControl/Private/Tools/DirectMeshControlTool.cpp
new file mode 100644
--- /dev/null
+++ b/Engine/Plugins/Experimental/Animation/DirectMeshControl/Source/DirectMeshControl/Private/Tools/DirectMeshControlTool.cpp
@@ -0,0 +1,961 @@
+// Copyright Epic Games, Inc. All Rights Reserved.
+
+
+#include "Tools/DirectMeshControlTool.h"
+
+#include "Materials/Material.h"
+#include "Materials/MaterialInstanceDynamic.h"
+
+#include "ContextObjectStore.h"
+#include "DirectMeshControlComponent.h"
+#include "DirectMeshControlUtilities.h"
+#include "InteractiveToolManager.h"
+#include "IPersonaEditorModeManager.h"
+#include "ToolTargetManager.h"
+#include "PreviewMesh.h"
+#include "ModelingToolTargetUtil.h"
+#include "PersonaModule.h"
+#include "SkeletalDebugRendering.h"
+#include "SkeletalMeshAttributes.h"
+#include "SkeletalRenderPublic.h"
+#include "ToolSetupUtil.h"
+#include "BaseBehaviors/ClickDragBehavior.h"
+#include "BaseBehaviors/MouseHoverBehavior.h"
+#include "Behaviors/MultiButtonClickDragBehavior.h"
+#include "Polygroups/PolygroupSet.h"
+#include "Polygroups/PolygroupUtil.h"
+#include "Selection/PolygonSelectionMechanic.h"
+#include "SkeletalMesh/SkinWeightsPaintTool.h"
+#include "TargetInterfaces/PrimitiveComponentBackedTarget.h"
+#include "Util/ColorConstants.h"
+#include "Components/SkeletalMeshComponent.h"
+#include "Rendering/SkeletalMeshRenderData.h"
+#include "Selection/GeometrySelector.h"
+#include "Animation/MeshDeformer.h"
+#include "RenderResource.h"
+#include "RHI.h"
+#include "SkeletalMeshLODRenderDataMeshAdapter.h"
+#include "ToolBuilderUtil.h"
+#include "DynamicMesh/DynamicVertexSkinWeightsAttribute.h"
+#include "TargetInterfaces/SkeletalMeshBackedTarget.h"
+#include "DynamicMeshToMeshDescription.h"
+#include "MeshDescription.h"
+#include "DynamicSubmesh3.h"
+#include "EngineUtils.h"
+#include "FileHelpers.h"
+#include "LODUtilities.h"
+#include "OptimusDeformer.h"
+#include "OptimusDeformerInstance.h"
+#include "OptimusResourceDescription.h"
+#include "SkeletalMeshAttributes.h"
+#include "StaticToSkeletalMeshConverter.h"
+#include "Animation/SkeletalMeshActor.h"
+#include "Animation/Skeleton.h"
+#include "Components/Viewport.h"
+#include "Rendering/SkeletalMeshLODModel.h"
+#include "Rendering/SkeletalMeshModel.h"
+
+#define LOCTEXT_NAMESPACE "UDirectMeshControlTool"
+
+/**
+ * UDirectMeshControlToolBuilder
+ */
+
+bool UDirectMeshControlToolBuilder::CanBuildTool(const FToolBuilderState& SceneState) const
+{
+	return SceneState.TargetManager->CountSelectedAndTargetable(SceneState, GetTargetRequirements()) == 1;
+}
+
+const FToolTargetTypeRequirements& UDirectMeshControlToolBuilder::GetTargetRequirements() const
+{
+	static const FToolTargetTypeRequirements TypeRequirements(UPrimitiveComponentBackedTarget::StaticClass());
+	return TypeRequirements;
+}
+
+UInteractiveTool* UDirectMeshControlToolBuilder::BuildTool(const FToolBuilderState& SceneState) const
+{
+	UDirectMeshControlTool* NewTool = NewObject<UDirectMeshControlTool>(SceneState.ToolManager);
+	
+	UToolTarget* Target = SceneState.TargetManager->BuildFirstSelectedTargetable(SceneState, GetTargetRequirements());
+	NewTool->SetTarget(Target);
+	NewTool->Init(SceneState);
+	
+	return NewTool;
+}
+
+namespace PolyGroupTool
+{
+	
+void CleanComponent(USceneComponent* Component)
+{
+	if (Component)
+	{
+		Component->DetachFromComponent(FDetachmentTransformRules::KeepRelativeTransform);
+		Component->UnregisterComponent();
+		Component->SetVisibility(false);
+		Component->DestroyComponent();
+	}		
+}
+
+}
+
+/**
+ * UDirectMeshControlTool
+ */
+
+void UDirectMeshControlTool::Init(const FToolBuilderState& InSceneState)
+{
+	TargetWorld = InSceneState.World;
+
+	const UContextObjectStore* ContextObjectStore = InSceneState.ToolManager->GetContextObjectStore();
+	EditorContext = ContextObjectStore->FindContext<USkeletalMeshEditorContextObjectBase>();
+	PersonaModeManagerContext = ContextObjectStore->FindContext<UPersonaEditorModeManagerContext>();
+}
+
+void UDirectMeshControlTool::Setup()
+{
+	Super::Setup();
+
+	ISkeletalMeshBackedTarget* SkeletalMeshBackedTarget = Cast<ISkeletalMeshBackedTarget>(Target); 
+	SourceSkeletalMesh = SkeletalMeshBackedTarget ? SkeletalMeshBackedTarget->GetSkeletalMesh() : nullptr;
+	if (!SourceSkeletalMesh.Get())
+	{
+		return;
+	}
+
+	FToolBuilderState State; GetToolManager()->GetContextQueriesAPI()->GetCurrentSelectionState(State);
+	UActorComponent* ActorComponent = ToolBuilderUtil::FindFirstComponent(State, [&](UActorComponent* Component)
+	{
+		return IsValid(Component) && Component->IsA<USkeletalMeshComponent>();
+	});
+	SourceSkeletalMeshComponent = Cast<USkeletalMeshComponent>(ActorComponent);
+	if (!SourceSkeletalMeshComponent.Get())
+	{
+		return;
+	}
+	
+	SourceSkeletalMeshComponent->SetAlwaysUseMeshDeformer(true);
+	SourceSkeletalMeshComponent->SetForcedLOD(1);
+
+	SetupProperties();
+	SetupPreviewMesh();
+	SetupBehaviors();
+	SetupPolyGroupSelection();
+	
+	if (EditorContext.IsValid())
+	{
+		EditorContext->BindTo(this);
+		EditorContext->ToggleBoneManipulation(true);
+	}
+}
+
+void UDirectMeshControlTool::SetupPreviewMesh()
+{
+	PreviewMesh = NewObject<UPreviewMesh>(this);
+	PreviewMesh->bBuildSpatialDataStructure = true;
+	{
+		PreviewMesh->CreateInWorld(TargetWorld.Get(), FTransform::Identity);
+		if (SourceSkeletalMeshComponent.IsValid())
+		{	
+			PreviewMesh->GetActor()->AttachToComponent(SourceSkeletalMeshComponent.Get(), FAttachmentTransformRules::KeepWorldTransform);
+		}
+	}
+
+	FTransform WorldTransform = UE::ToolTarget::GetLocalToWorldTransform(Target);
+	PreviewMesh->SetTransform(WorldTransform);
+
+	ToolSetupUtil::ApplyRenderingConfigurationToPreview(PreviewMesh, Target);
+	PreviewMesh->SetTangentsMode(EDynamicMeshComponentTangentsMode::AutoCalculated);
+	PreviewMesh->ReplaceMesh(UE::ToolTarget::GetDynamicMeshCopy(Target));
+
+	// probably not needed
+	PreviewMesh->SetMaterials({UMaterialInstanceDynamic::Create(UE::DMC::GetMaterial(), PreviewMesh->GetActor())});
+	const FColor DummyColor(0);
+	PreviewMesh->SetTriangleColorFunction([DummyColor](const FDynamicMesh3*, int)
+	{
+		return DummyColor;
+	}, UPreviewMesh::ERenderUpdateMode::FastUpdate);
+
+	// configure secondary render material for selected triangles
+	UMaterialInterface* SelectionMaterial = ToolSetupUtil::GetSelectionMaterial(FLinearColor::Yellow, GetToolManager());
+	if (SelectionMaterial != nullptr)
+	{
+		PreviewMesh->SetSecondaryRenderMaterial(SelectionMaterial);
+	}
+
+	// FIXME: This setting really belongs on the underlying mesh.
+	UDynamicMeshComponent* DynamicMeshComponent = Cast<UDynamicMeshComponent>(PreviewMesh->GetRootComponent());
+	DynamicMeshComponent->SetVertexColorSpaceTransformMode(EDynamicMeshVertexColorTransformMode::LinearToSRGB);
+
+	if (SourceSkeletalMesh.IsValid())
+	{
+		if (Properties->bSubMeshes)
+		{
+			GenerateDMCSubSkeletalMeshes();
+		}
+		UpdateVisualization();
+	}
+	
+	PreviewMesh->SetVisible(false);
+}
+
+void UDirectMeshControlTool::CleanDMCSubSkeletalMeshes()
+{
+	for (TWeakObjectPtr<USkeletalMeshComponent>& SubComponent: SubComponents)
+	{
+		if (SubComponent.IsValid())
+		{
+			AActor* SubMeshActor = Properties->bAttachToActor ? SubComponent->GetOwner() : nullptr;
+			PolyGroupTool::CleanComponent(SubComponent.Get());
+			SubComponent = nullptr;
+				
+			if (SubMeshActor)
+			{
+				SubMeshActor->Destroy();
+				SubMeshActor = nullptr;
+			}
+		}
+	}
+	SubComponents.Reset();
+}
+
+void UDirectMeshControlTool::GenerateDMCSubSkeletalMeshes()
+{
+	CleanDMCSubSkeletalMeshes();
+	
+	using namespace UE::Geometry;
+
+	const FDynamicMesh3* Mesh = PreviewMesh->GetMesh();
+	static const FName LayerName("dmc-polygroup");
+	
+	const FGroupSubMeshes& GroupSubMeshes = UE::DMC::GetSubMeshes(SourceSkeletalMesh.Get(), Mesh, LayerName);
+
+	UOptimusDeformer* MeshDeformer = UE::DMC::GetDeformer();
+	
+	AActor* SourceSKMActor = SourceSkeletalMeshComponent->GetOwner();
+
+	const TArray<TObjectPtr<USkeletalMesh>>& SubMeshes = GroupSubMeshes.GetSubSkeletalMeshes();
+	for (const auto& [GroudID, SubMeshIndex] : GroupSubMeshes.GetGroupIdToSubMesh())
+	{
+		USkeletalMesh* SubSkeletalMesh = SubMeshes.IsValidIndex(SubMeshIndex) ? SubMeshes[SubMeshIndex].Get() : nullptr;
+		if (ensure(SubSkeletalMesh))
+		{
+			USkeletalMeshComponent* SubSkeletalMeshComponent;
+			if (Properties->bAttachToActor)
+			{
+				AActor* SubMeshActor = PreviewMesh->GetWorld()->SpawnActor<AActor>();
+				
+				USceneComponent* RootComp = NewObject<USceneComponent>(SubMeshActor);
+				RootComp->RegisterComponent();
+				SubMeshActor->SetRootComponent(RootComp);
+				
+				//SubSkeletalMeshComponent = NewObject<USkeletalMeshComponent>(SubMeshActor);
+				SubSkeletalMeshComponent = NewObject<UDirectMeshControlComponent>(SubMeshActor);
+	
+				SubSkeletalMeshComponent->SetSkeletalMesh(SubSkeletalMesh);
+				SubSkeletalMeshComponent->RegisterComponent();
+				SubSkeletalMeshComponent->AttachToComponent(RootComp, FAttachmentTransformRules::KeepWorldTransform);
+				
+				SubMeshActor->AttachToActor(SourceSKMActor, FAttachmentTransformRules::KeepWorldTransform);
+				
+				if (UE::DMC::UseMaterialHover())
+				{
+					UMaterialInstanceDynamic* MaterialInstance = UMaterialInstanceDynamic::Create(UE::DMC::GetMaterial(), SubMeshActor);
+					MaterialInstance->SetScalarParameterValue(TEXT("Hovered"), 1.f);
+					SubSkeletalMeshComponent->SetMaterial(0, MaterialInstance);
+				}
+			}
+			else
+			{
+				//SubSkeletalMeshComponent = NewObject<USkeletalMeshComponent>(SourceSKMActor);
+				SubSkeletalMeshComponent = NewObject<UDirectMeshControlComponent>(SourceSKMActor);
+				SubSkeletalMeshComponent->SetSkeletalMesh(SubSkeletalMesh);
+				SubSkeletalMeshComponent->AttachToComponent(SourceSkeletalMeshComponent.Get(), FAttachmentTransformRules::KeepWorldTransform);
+				SubSkeletalMeshComponent->RegisterComponent();
+			}
+
+			SubSkeletalMeshComponent->SetMeshDeformer(MeshDeformer);
+			if (!UE::DMC::UseMaterialHover())
+			{
+				UOptimusDeformerInstance* MeshDeformerInstance = Cast<UOptimusDeformerInstance>(SubSkeletalMeshComponent->GetMeshDeformerInstance());
+				if (ensure(MeshDeformerInstance))
+				{
+					FColor Color = LinearColors::SelectFColor(GroudID);
+					Color.A = 255.f;
+					ensure(MeshDeformerInstance->SetLinearColorVariable(UE::DMC::GetColorVarName(), FLinearColor(Color)));
+				}
+			}
+			
+			SubSkeletalMeshComponent->SetRelativeLocation(Properties->OffsetSub);
+			const int32 Index = SubComponents.Add(SubSkeletalMeshComponent);
+			PolyGroupToSubMesh.Emplace(GroudID, Index);
+		}
+	}
+}
+
+void UDirectMeshControlTool::SetupProperties()
+{
+	ToolPropertyObjects.Add(this);
+	
+	Properties = NewObject<UDirectMeshControlProperties>(this);
+	Properties->Initialize(this);
+	Properties->RestoreProperties(this);
+	AddToolPropertySource(Properties);
+
+	Properties->WatchProperty(Properties->Alpha, [this](float) { UpdateVisualization(); });
+	Properties->WatchProperty(Properties->Offset, [this](FVector InOffset)
+	{
+		FTransform Transform = PreviewMesh->GetTransform();
+		Transform.SetLocation(InOffset);
+		PreviewMesh->SetTransform(Transform);
+	});
+
+	Properties->WatchProperty(Properties->OffsetSub, [this](FVector InOffset)
+	{
+		for (TWeakObjectPtr<USkeletalMeshComponent>& SubComponent: SubComponents)
+		{
+			if (SubComponent.IsValid())
+			{
+				SubComponent->SetRelativeLocation(InOffset);
+			}
+		}
+	});
+	
+	Properties->WatchProperty(Properties->Action, [this](EDMCAction)
+	{
+		if (IPersonaEditorModeManager* ModeManager = PersonaModeManagerContext.IsValid() ? PersonaModeManagerContext->GetPersonaEditorModeManager() : nullptr)
+		{
+			if (Properties->Action == EDMCAction::None || Properties->Action == EDMCAction::Pose)
+			{
+				ModeManager->ActivateMode(FPersonaEditModes::SkeletonSelection);	
+			}
+			else
+			{
+				ModeManager->DeactivateMode(FPersonaEditModes::SkeletonSelection);	
+			}
+		}
+		
+		PreviewMesh->FastNotifySecondaryTrianglesChanged();
+		UpdateVisualization();
+	});
+	
+	Properties->WatchProperty(Properties->bSubMeshes, [this](bool bNewValue)
+	{
+		if (PreviewMesh)
+		{	
+			if (bNewValue)
+			{
+				GenerateDMCSubSkeletalMeshes();
+			}
+			else
+			{
+				CleanDMCSubSkeletalMeshes();
+			}
+		}
+	});
+}
+
+void UDirectMeshControlTool::SetupBehaviors()
+{
+	static constexpr int PRE_INDIRECT_GIZMO_PRIORITY = FInputCapturePriority::DEFAULT_GIZMO_PRIORITY-2;
+	
+	// Add middle mouse input behavior for indirect manipulation
+	ULocalClickDragInputBehavior* MiddleClickDragBehavior = NewObject<ULocalClickDragInputBehavior>(this);
+	MiddleClickDragBehavior->Initialize();
+	MiddleClickDragBehavior->SetDefaultPriority(PRE_INDIRECT_GIZMO_PRIORITY);
+	MiddleClickDragBehavior->SetUseMiddleMouseButton();
+
+	MiddleClickDragBehavior->CanBeginClickDragFunc = [this](const FInputDeviceRay& InRay)
+	{
+		if (Properties->Action == EDMCAction::Pose)
+		{
+			if (CurrentBone == INDEX_NONE)
+			{
+				const int32 HoveredPolyGroup = GetHoveredPolyGroup(InRay);
+				if (HoveredPolyGroup != CurrentPolyGroup)
+				{
+					CurrentPolyGroup = HoveredPolyGroup;
+					UpdateVisualization();
+				}
+
+				const int32* BoneFromPolyGroup = CurrentPolyGroup != INDEX_NONE ? PolyGroupToBone.Find(CurrentPolyGroup) : nullptr;
+				CurrentBone = BoneFromPolyGroup ? *BoneFromPolyGroup : INDEX_NONE;
+			}
+
+			if (CurrentBone != INDEX_NONE)
+			{
+				if (SourceSkeletalMesh.IsValid())
+				{
+					bRebuild = true;
+					const FReferenceSkeleton& ReferenceSkeleton = SourceSkeletalMesh->GetRefSkeleton();
+					GetNotifier()->Notify({ReferenceSkeleton.GetBoneName(CurrentBone)}, ESkeletalMeshNotifyType::BonesSelected);
+				}
+			}
+			else
+			{
+				GetNotifier()->Notify({}, ESkeletalMeshNotifyType::BonesSelected);
+			}
+		}
+		
+		static const FInputRayHit InvalidRayHit;
+		return InvalidRayHit;
+	};
+	AddInputBehavior(MiddleClickDragBehavior);
+
+	// Add left/right mouse input behavior for indirect manipulation
+	UMultiButtonClickDragBehavior* MultiIndirectClickDragBehavior = NewObject<UMultiButtonClickDragBehavior>(this);
+	MultiIndirectClickDragBehavior->Initialize();
+	MultiIndirectClickDragBehavior->SetDefaultPriority(PRE_INDIRECT_GIZMO_PRIORITY);
+	MultiIndirectClickDragBehavior->EnableButton(EKeys::LeftMouseButton);
+	MultiIndirectClickDragBehavior->EnableButton(EKeys::MiddleMouseButton);
+	MultiIndirectClickDragBehavior->EnableButton(EKeys::RightMouseButton);
+	MultiIndirectClickDragBehavior->ModifierCheckFunc = FInputDeviceState::IsCtrlKeyDown;
+	MultiIndirectClickDragBehavior->CanBeginClickDragFunc = [this](const FInputDeviceRay& InRay)
+	{
+		if (Properties->Action == EDMCAction::Pose)
+		{
+			if (CurrentBone == INDEX_NONE)
+			{
+				const int32 HoveredPolyGroup = GetHoveredPolyGroup(InRay);
+				if (HoveredPolyGroup != CurrentPolyGroup)
+				{
+					CurrentPolyGroup = HoveredPolyGroup;
+					UpdateVisualization();
+				}
+
+				const int32* BoneFromPolyGroup = CurrentPolyGroup != INDEX_NONE ? PolyGroupToBone.Find(CurrentPolyGroup) : nullptr;
+				CurrentBone = BoneFromPolyGroup ? *BoneFromPolyGroup : INDEX_NONE;
+			}
+
+			if (CurrentBone != INDEX_NONE)
+			{
+				if (SourceSkeletalMesh.IsValid())
+				{
+					bRebuild = true;
+					const FReferenceSkeleton& ReferenceSkeleton = SourceSkeletalMesh->GetRefSkeleton();
+					GetNotifier()->Notify({ReferenceSkeleton.GetBoneName(CurrentBone)}, ESkeletalMeshNotifyType::BonesSelected);
+				}
+			}
+			else
+			{
+				GetNotifier()->Notify({}, ESkeletalMeshNotifyType::BonesSelected);
+			}
+		}
+		
+		static const FInputRayHit InvalidRayHit;
+		return InvalidRayHit;
+	};
+	AddInputBehavior(MultiIndirectClickDragBehavior);
+
+	// Add hover behavior
+	UMouseHoverBehavior* HoverBehavior = NewObject<UMouseHoverBehavior>(this);	
+	HoverBehavior->Initialize(this);
+	HoverBehavior->SetDefaultPriority(PRE_INDIRECT_GIZMO_PRIORITY);
+	AddInputBehavior(HoverBehavior, this);
+}
+
+void UDirectMeshControlTool::SetupPolyGroupSelection()
+{
+	using namespace UE::Geometry;
+	
+	if (!PreviewMesh)
+	{
+		return;
+	}
+	
+	// set up vertex selection mechanic
+	SelectionMechanic = NewObject<UPolygonSelectionMechanic>(this);
+	SelectionMechanic->Setup(this);
+	SelectionMechanic->SetIsEnabled(false);
+	// the property is disabled and managed thru customization
+	SetToolPropertySourceEnabled(SelectionMechanic->Properties, false);
+	SelectionMechanic->SetSelectionDragToolUpdateType(ESelectionDragToolUpdateType::OnRelease);
+	SelectionMechanic->Properties->bDisplayPolygroupReliantControls = false;
+
+	SelectionMechanic->Properties->bSelectVertices = false;
+	SelectionMechanic->Properties->bSelectEdges = false;
+	SelectionMechanic->Properties->bSelectFaces = true;
+	SelectionMechanic->SetShowSelectableCorners(false);
+
+	// adjust selection rendering for this context
+	SelectionMechanic->PolyEdgesRenderer.PointColor = FLinearColor(0.78f, 0.f, 0.78f);
+	SelectionMechanic->PolyEdgesRenderer.PointSize = 5.0f;
+	SelectionMechanic->PolyEdgesRenderer.DepthBias = 0.01f;
+	
+	SelectionMechanic->HilightRenderer.PointColor = FLinearColor::Red;
+	SelectionMechanic->HilightRenderer.PointSize = 10.0f;
+	
+	SelectionMechanic->SelectionRenderer.LineThickness = 0.0f;
+	SelectionMechanic->SelectionRenderer.DepthBias = 0.01f;
+	SelectionMechanic->SelectionRenderer.PointColor = FLinearColor::Yellow;
+	SelectionMechanic->SelectionRenderer.PointSize = 5.0f;
+	SelectionMechanic->SetShowEdges(false);
+
+	const FDynamicMesh3* Mesh = PreviewMesh->GetMesh();
+	static const FName LayerName("dmc-polygroup");
+	const FDynamicMeshTriangleLabelAttribute* LabelsLayer = FindTriangleLabelLayerByName(*Mesh, LayerName);
+	if (ensure(LabelsLayer))
+	{
+		TMap<FName, int32> BoneToID;
+		PolygroupAttrib = MakeUnique<FDynamicMeshPolygroupAttribute>(const_cast<FDynamicMesh3*>(Mesh));
+		
+		for (int32 Index = 0; Index < Mesh->TriangleCount(); Index++)
+		{
+			const FName BoneName = LabelsLayer->GetValue(Index);
+			if (int32* GroupID = BoneToID.Find(BoneName))
+			{
+				PolygroupAttrib->SetValue(Index, *GroupID);	
+			}
+			else
+			{
+				PolygroupAttrib->SetValue(Index, BoneToID.Num());
+				BoneToID.Emplace(BoneName, BoneToID.Num());
+			}
+		}
+		
+		const FReferenceSkeleton& ReferenceSkeleton = SourceSkeletalMesh->GetRefSkeleton();
+		for (const auto& [BoneName, GroupID]: BoneToID)
+		{
+			PolyGroupToBone.Emplace(GroupID, ReferenceSkeleton.FindBoneIndex(BoneName));
+		}
+	}
+
+	// initialize the polygon selection mechanic
+	PreviewMesh->ProcessMesh([this](const FDynamicMesh3& Mesh)
+	{
+		Topology = MakeUnique<FGroupTopology>(&Mesh, PolygroupAttrib.Get(), true);
+		ActiveGroupSet = PolygroupAttrib.IsValid() ? MakeUnique<FPolygroupSet>(&Mesh, PolygroupAttrib.Get()) : MakeUnique<FPolygroupSet>(&Mesh);
+		MeshSpatial = MakeShared<FDynamicMeshAABBTree3>();
+		MeshSpatial->SetMesh(&Mesh);
+		SelectionMechanic->Initialize(&Mesh, FTransform::Identity, TargetWorld.Get(), Topology.Get(), [this]() { return &GetSpatial(); });
+	});
+
+	// triangles rendering
+	PreviewMesh->EnableSecondaryTriangleBuffers([this](const FDynamicMesh3* Mesh, int32 TriangleID)
+	{
+		return Properties->Action == EDMCAction::Select &&
+				SelectionMechanic->Properties->bSelectFaces &&
+				SelectionMechanic->GetActiveSelection().IsSelectedTriangle(Mesh, Topology.Get(), TriangleID);
+	});
+	SelectionMechanic->OnSelectionChanged.AddWeakLambda(this, [this]()
+	{
+		PreviewMesh->FastNotifySecondaryTrianglesChanged();
+	});
+	SelectionMechanic->OnFaceSelectionPreviewChanged.AddWeakLambda(this, [this]()
+	{
+		PreviewMesh->FastNotifySecondaryTrianglesChanged();
+	});
+}
+
+void UDirectMeshControlTool::Shutdown(EToolShutdownType ShutdownType)
+{
+	if (ShutdownType == EToolShutdownType::Accept)
+	{
+		// TBD
+	}
+	
+	const bool bDetachActor = Properties->bAttachToActor;
+	
+	
+	Super::Shutdown(ShutdownType);
+	
+	if (SelectionMechanic)
+    {
+    	SelectionMechanic->Shutdown();
+    	SelectionMechanic = nullptr;
+    }
+	
+	// remove preview mesh
+	if (PreviewMesh != nullptr)
+	{
+		PreviewMesh->SetVisible(false);
+		PreviewMesh->Disconnect();
+		PreviewMesh = nullptr;
+	}
+
+	{
+		CleanDMCSubSkeletalMeshes();
+	}
+
+	// save properties	
+	Properties->SaveProperties(this);
+	
+	if (EditorContext.IsValid())
+	{
+		EditorContext->ShowSkeleton();
+		EditorContext->UnbindFrom(this);
+	}
+}
+
+void UDirectMeshControlTool::RegisterActions(FInteractiveToolActionSet& InOutActionSet)
+{
+	Super::RegisterActions(InOutActionSet);
+}
+
+void UDirectMeshControlTool::OnClickPress(const FInputDeviceRay& InRay)
+{
+}
+
+void UDirectMeshControlTool::OnClickDrag(const FInputDeviceRay& InDragPos)
+{
+}
+
+void UDirectMeshControlTool::OnClickRelease(const FInputDeviceRay& InReleasePos)
+{
+}
+
+void UDirectMeshControlTool::OnTerminateDragSequence()
+{}
+
+int32 UDirectMeshControlTool::GetHoveredPolyGroup(const FInputDeviceRay& InRay) const
+{
+	if (PreviewMesh)
+	{
+		if (Properties && Properties->bHitProxy)
+		{
+			if (LastScreenPosition != InRay.ScreenPosition)
+			{
+				LastScreenPosition = InRay.ScreenPosition;
+				LastHit = INDEX_NONE;
+			
+				FViewport* Viewport = GetToolManager()->GetContextQueriesAPI()->GetFocusedViewport();
+				HHitProxy* HitProxy = Viewport ? Viewport->GetHitProxy(InRay.ScreenPosition[0], InRay.ScreenPosition[1]) : nullptr;
+			
+				if (const HActor* ActorProxy = HitProxyCast<HActor>(HitProxy))
+				{
+					if (const UPrimitiveComponent* PrimitiveComponent = ActorProxy->PrimComponent)
+					{
+						for (const auto&[GroupId, SubMeshIndex]: PolyGroupToSubMesh)
+						{
+							USkeletalMeshComponent* SkeletalMeshComponent =
+								SubComponents.IsValidIndex(SubMeshIndex) ? SubComponents[SubMeshIndex].Get() : nullptr;
+							if (SkeletalMeshComponent && SkeletalMeshComponent == PrimitiveComponent)
+							{
+								LastHit = GroupId; 
+								break;
+							}
+						}
+					}
+				}
+				
+				return LastHit;
+			}
+			
+			return LastHit;
+		}
+
+		FHitResult Hit;
+		const FRay& WorldRay = InRay.WorldRay;
+		
+		PreviewMesh->SetVisible(true);
+		if (PreviewMesh->FindRayIntersection(WorldRay, Hit))
+		{
+			PreviewMesh->SetVisible(false);
+			return ActiveGroupSet->GetGroup(Hit.FaceIndex);
+		}
+		PreviewMesh->SetVisible(false);
+	}
+	return INDEX_NONE;
+}
+
+FInputRayHit UDirectMeshControlTool::BeginHoverSequenceHitTest(const FInputDeviceRay& InRay)
+{
+	if (bRebuild)
+	{
+		const bool bRebuildSpatial = Properties->Action == EDMCAction::Pose;
+		PreviewMesh->NotifyDeferredEditCompleted(UPreviewMesh::ERenderUpdateMode::FastUpdate, EMeshRenderAttributeFlags::Positions, bRebuildSpatial);
+		bRebuild = false;
+	}
+	
+	static const FInputRayHit InvalidRayHit;
+	if (Properties->Action != EDMCAction::Pose)
+	{
+		return InvalidRayHit;
+	}
+
+	const int32 PolyGroup = GetHoveredPolyGroup(InRay);
+	if (PolyGroup == INDEX_NONE)
+	{
+		if (PolyGroup != CurrentPolyGroup)
+		{
+			CurrentPolyGroup = PolyGroup;
+			UpdateVisualization();
+		}
+		CurrentBone = INDEX_NONE;
+		GetNotifier()->Notify({}, ESkeletalMeshNotifyType::BonesSelected);
+		
+		return InvalidRayHit;
+	}
+
+	const int32* BoneIndex = PolyGroupToBone.Find(PolyGroup);
+	if (!BoneIndex)
+	{
+		return InvalidRayHit;
+	}
+	
+	FInputRayHit RayHit;
+	RayHit.bHit = true;
+	return RayHit;
+}
+
+void UDirectMeshControlTool::OnBeginHover(const FInputDeviceRay& InRay)
+{
+	CurrentPolyGroup = INDEX_NONE;
+	CurrentBone = INDEX_NONE;
+	
+	if (Properties->Action == EDMCAction::Pose)
+	{
+		GetNotifier()->Notify({}, ESkeletalMeshNotifyType::BonesSelected);
+		UpdateVisualization();
+	}
+}
+
+bool UDirectMeshControlTool::OnUpdateHover(const FInputDeviceRay& InRay)
+{
+	if (Properties->Action != EDMCAction::Pose)
+	{
+		return false;
+	}
+	
+	const int32 HoveredPolyGroup = GetHoveredPolyGroup(InRay);
+	if (HoveredPolyGroup != CurrentPolyGroup)
+	{
+		CurrentPolyGroup = HoveredPolyGroup;
+		UpdateVisualization();
+	}
+
+	const int32* BoneFromPolyGroup = CurrentPolyGroup != INDEX_NONE ? PolyGroupToBone.Find(CurrentPolyGroup) : nullptr;
+	const int32 BoneIndex = BoneFromPolyGroup ? *BoneFromPolyGroup : INDEX_NONE;
+	if (BoneIndex != CurrentBone)
+	{
+		CurrentBone = BoneIndex;
+		// if (CurrentBone != INDEX_NONE)
+		// {
+		// 	const FReferenceSkeleton& ReferenceSkeleton = WeakMesh->GetRefSkeleton();
+		// 	GetNotifier()->Notify({ReferenceSkeleton.GetBoneName(CurrentBone)}, ESkeletalMeshNotifyType::BonesSelected);
+		// }
+		// else
+		{
+			GetNotifier()->Notify({}, ESkeletalMeshNotifyType::BonesSelected);
+		}
+	}
+	
+	return CurrentPolyGroup != INDEX_NONE;
+}
+
+void UDirectMeshControlTool::OnEndHover()
+{}
+
+void UDirectMeshControlTool::OnTick(float DeltaTime)
+{
+	if (SelectionMechanic)
+	{
+		SelectionMechanic->SetIsEnabled(Properties->Action == EDMCAction::Select);
+		SelectionMechanic->Tick(DeltaTime);
+	}
+}
+
+FInputRayHit UDirectMeshControlTool::CanBeginClickDragSequence(const FInputDeviceRay& InRay)
+{
+	ensure(false);
+	static const FInputRayHit InvalidRayHit;
+	return InvalidRayHit;
+}
+
+void UDirectMeshControlTool::HandleSkeletalMeshModified(const TArray<FName>& InBoneNames, const ESkeletalMeshNotifyType InNotifyType)
+{
+	if (GetNotifier()->Notifying())
+	{
+		return;
+	}
+
+	switch (InNotifyType)
+	{
+		case ESkeletalMeshNotifyType::BonesAdded:
+			Selection = InBoneNames;
+			break;
+		case ESkeletalMeshNotifyType::BonesRemoved:
+			Selection.RemoveAll([&InBoneNames](const FName& BoneName)
+			{
+				return InBoneNames.Contains(BoneName);
+			});
+			break;
+		case ESkeletalMeshNotifyType::BonesMoved:
+			break;
+		case ESkeletalMeshNotifyType::BonesSelected:
+			Selection.Reset();
+			Selection = InBoneNames;
+			break;
+		case ESkeletalMeshNotifyType::BonesRenamed:
+			Selection = InBoneNames;
+			break;
+		case ESkeletalMeshNotifyType::HierarchyChanged:
+			break;
+		default:
+			break;
+	}
+}
+
+void UDirectMeshControlTool::DrawHUD(FCanvas* Canvas, IToolsContextRenderAPI* RenderAPI)
+{
+	Super::DrawHUD(Canvas, RenderAPI);
+
+	if (SelectionMechanic && Properties->Action == EDMCAction::Select)
+	{
+		SelectionMechanic->DrawHUD(Canvas, RenderAPI);
+	}
+}
+
+void UDirectMeshControlTool::Render(IToolsContextRenderAPI* RenderAPI)
+{
+	if (SelectionMechanic && Properties->Action == EDMCAction::Select)
+	{
+		SelectionMechanic->Render(RenderAPI);
+	}
+}
+
+void UDirectMeshControlTool::OnPropertyModified(UObject* PropertySet, FProperty* Property)
+{
+	Super::OnPropertyModified(PropertySet, Property);
+}
+
+FBox UDirectMeshControlTool::GetWorldSpaceFocusBox()
+{
+	return USingleSelectionTool::GetWorldSpaceFocusBox();
+}
+
+FName UDirectMeshControlTool::GetCurrentBone() const
+{
+	return Selection.IsEmpty() ? NAME_None : Selection.Last(); 
+}
+
+TArray<int32> UDirectMeshControlTool::GetSelectedBoneIndexes() const
+{
+	TArray<int32> Indexes;
+
+	if (SourceSkeletalMesh.IsValid())
+	{
+		const FReferenceSkeleton& ReferenceSkeleton = SourceSkeletalMesh->GetRefSkeleton();
+		Algo::Transform(Selection, Indexes, [&](const FName& BoneName)
+		{
+			return ReferenceSkeleton.FindRawBoneIndex(BoneName);
+		});
+	}
+	
+	return Indexes;
+}
+
+FDynamicMeshAABBTree3& UDirectMeshControlTool::GetSpatial()
+{
+	if (bSpatialDirty)
+	{
+		MeshSpatial->Build();
+		bSpatialDirty = false;
+	}
+	return *MeshSpatial;
+}
+
+void UDirectMeshControlTool::UpdateVisualization()
+{
+	if (!PreviewMesh)
+	{
+		return;
+	}
+
+	// {
+	// 	TSet<int> Triangles;
+	// 	if (ActiveGroupSet)
+	// 	{
+	// 		PreviewMesh->SetTriangleColorFunction([this, &Triangles](const FDynamicMesh3* Mesh, int TriangleID)
+	// 	   {
+	// 		   const int32 GroupId = ActiveGroupSet->GetGroup(TriangleID);
+	// 		   FColor Color = UE::Geometry::LinearColors::SelectFColor(GroupId);
+	//
+	// 		   // const int32* BoneIndex = PolyGroupToBone.Find(GroupId);
+	// 		   // if (BoneIndex)
+	// 		   // {
+	// 		   // 	FLinearColor BoneColor = SkeletalDebugRendering::GetSemiRandomColorForBone(*BoneIndex);
+	// 		   // 	Color = BoneColor.ToFColor(false);
+	// 		   // }
+	// 		
+	// 		   Color.A = (uint8)(0.5f + Properties->Alpha * 255.f);
+	// 		
+	// 		   if (Properties->Action == EDMCAction::Pose)
+	// 		   {
+	// 			   if (/*!BoneIndex || */GroupId != CurrentPolyGroup)
+	// 			   {
+	// 				   Color.A = 0;
+	// 			   }
+	// 		   }
+	// 		   return Color;
+	// 	   }, 
+	// 	   UPreviewMesh::ERenderUpdateMode::FastUpdate);
+	// 	}
+	// }
+
+	{
+		
+		
+		const bool bPosing = Properties->Action == EDMCAction::Pose;
+		const float Alpha = Properties->Alpha;
+		for (const auto&[GroupId, SubMeshIndex]: PolyGroupToSubMesh)
+		{
+			if (SubComponents.IsValidIndex(SubMeshIndex) && SubComponents[SubMeshIndex].IsValid())
+			{
+				USkeletalMeshComponent* SkeletalMeshComponent = SubComponents[SubMeshIndex].Get();
+				UOptimusDeformerInstance* MeshDeformerInstance = Cast<UOptimusDeformerInstance>(SkeletalMeshComponent->GetMeshDeformerInstance());
+				if (ensure(MeshDeformerInstance))
+				{
+					const bool bCurrentGroup = GroupId == CurrentPolyGroup;
+					if (!UE::DMC::UseMaterialHover())
+					{
+						FColor Color = UE::Geometry::LinearColors::SelectFColor(GroupId);
+						Color.A = bPosing && !bCurrentGroup ? 0 : static_cast<uint8>(0.5f + Alpha * 255.f);
+						MeshDeformerInstance->SetLinearColorVariable(UE::DMC::GetColorVarName(), FLinearColor(Color));
+					}
+					else
+					{
+						if (UMaterialInstanceDynamic* MaterialInstance = Cast<UMaterialInstanceDynamic>(SkeletalMeshComponent->GetMaterial(0)))
+						{
+							MaterialInstance->SetScalarParameterValue(TEXT("Hovered"), bPosing && !bCurrentGroup ? 0.f : Alpha);
+						}
+					}
+				}
+			}
+		}
+	}
+}
+
+void UDirectMeshControlProperties::Initialize(UDirectMeshControlTool* InParentTool)
+{
+	ParentTool = InParentTool;
+}
+
+void UDirectMeshControlProperties::Assign()
+{
+	if (!ParentTool.IsValid() || !ParentTool->SourceSkeletalMesh.IsValid() || !ParentTool->PreviewMesh || !ParentTool->SelectionMechanic)
+	{
+		return;
+	}
+
+	const FReferenceSkeleton& RefSkeleton = ParentTool->SourceSkeletalMesh->GetRefSkeleton();
+	const int32 FindRawBoneIndex = RefSkeleton.FindRawBoneIndex(ParentTool->GetCurrentBone());
+
+	const TSet<int32> GroupIDs = ParentTool->SelectionMechanic->GetActiveSelection().SelectedGroupIDs;
+	for (int32 GroupId: GroupIDs)
+	{
+		ParentTool->PolyGroupToBone.FindOrAdd(GroupId) = FindRawBoneIndex;
+	}
+
+	ParentTool->GetNotifier()->Notify({}, ESkeletalMeshNotifyType::BonesSelected);
+	ParentTool->SelectionMechanic->ClearSelection();
+	ParentTool->UpdateVisualization();
+	ParentTool->PreviewMesh->FastNotifySecondaryTrianglesChanged();
+}
+
+void UDirectMeshControlProperties::DefaultAssign()
+{
+
+}
+
+#undef LOCTEXT_NAMESPACE
\ No newline at end of file
diff --git a/Engine/Plugins/Experimental/Animation/DirectMeshControl/Source/DirectMeshControl/Private/Tools/DirectMeshPolygroupTool.cpp b/Engine/Plugins/Experimental/Animation/DirectMeshControl/Source/DirectMeshControl/Private/Tools/DirectMeshPolygroupTool.cpp
new file mode 100644
--- /dev/null
+++ b/Engine/Plugins/Experimental/Animation/DirectMeshControl/Source/DirectMeshControl/Private/Tools/DirectMeshPolygroupTool.cpp
@@ -0,0 +1,668 @@
+// Copyright Epic Games, Inc. All Rights Reserved.
+
+#include "Tools/DirectMeshPolygroupTool.h"
+
+#include "BoneWeights.h"
+#include "InteractiveToolManager.h"
+#include "ToolBuilderUtil.h"
+#include "ToolSetupUtil.h"
+#include "ModelingToolTargetUtil.h"
+#include "Drawing/PreviewGeometryActor.h"
+#include "PreviewMesh.h"
+
+#include "MeshOpPreviewHelpers.h"
+#include "ModelingOperators.h"
+
+#include "DynamicMesh/DynamicMesh3.h"
+#include "DynamicMesh/MeshNormals.h"
+#include "DynamicMesh/MeshSharingUtil.h"
+#include "DynamicMeshEditor.h"
+#include "DynamicSubmesh3.h"
+#include "MeshRegionBoundaryLoops.h"
+#include "Util/ColorConstants.h"
+#include "Polygroups/PolygroupUtil.h"
+#include "Util/ColorConstants.h"
+#include "Selections/GeometrySelectionUtil.h"
+#include "Selection/StoredMeshSelectionUtil.h"
+#include "Selection/GeometrySelectionVisualization.h"
+#include "PropertySets/GeometrySelectionVisualizationProperties.h"
+#include "ToolTargetManager.h"
+#include "DynamicMesh/DynamicBoneAttribute.h"
+#include "DynamicMesh/DynamicVertexSkinWeightsAttribute.h"
+
+#include "CanvasTypes.h"
+#include "Engine/Engine.h"  // for GEngine->GetSmallFont()
+
+#include "Polygroups/PolygroupsGenerator.h"
+
+#include UE_INLINE_GENERATED_CPP_BY_NAME(DirectMeshPolygroupTool)
+
+using namespace UE::Geometry;
+
+#define LOCTEXT_NAMESPACE "UDirectMeshPolygroupTool"
+
+/*
+ * ToolBuilder
+ */
+
+USingleTargetWithSelectionTool* UDirectMeshPolygroupToolBuilder::CreateNewTool(const FToolBuilderState& SceneState) const
+{
+	return NewObject<UDirectMeshPolygroupTool>(SceneState.ToolManager);
+}
+
+bool UDirectMeshPolygroupToolBuilder::CanBuildTool(const FToolBuilderState& SceneState) const
+{
+	return USingleTargetWithSelectionToolBuilder::CanBuildTool(SceneState) &&
+		SceneState.TargetManager->CountSelectedAndTargetableWithPredicate(SceneState, GetTargetRequirements(),
+			[](UActorComponent& Component) { return !ToolBuilderUtil::IsVolume(Component); }) >= 1;
+}
+
+class FDirectMeshPolygroupsGenerator : public FPolygroupsGenerator
+{
+public:
+
+	struct FPolygroupsFromSkinWeightsParams
+	{
+		bool bTryUsingQuads = true;
+		bool bOptimizeNeighbours = false;
+		bool bRespectUVSeams = true;
+		bool bRespectNormalSeams = false;
+		double QuadAdjacencyWeight = 1.0;
+		double QuadMetricClamp = 1.0;
+		int MaxSearchRounds = 1;
+		TArray<int32> BonesToRemove;
+		TArray<int32> Parents;
+	};
+
+	FDirectMeshPolygroupsGenerator() {}
+	FDirectMeshPolygroupsGenerator(FDynamicMesh3* MeshIn)
+		: FPolygroupsGenerator(MeshIn) {}
+
+	/**
+	 * Find Polygroups based on skin weights
+	 */
+	bool FindPolygroupsFromSkinWeights(const FPolygroupsFromSkinWeightsParams& InParams = FPolygroupsFromSkinWeightsParams())
+	{
+		using namespace UE::AnimationCore;
+		
+		FoundPolygroups.Reset();
+		
+		if (!Mesh->HasAttributes())
+		{
+			return false;
+		}
+		
+		const FDynamicMeshVertexSkinWeightsAttribute* SkinWeights = Mesh->Attributes()->GetSkinWeightsAttribute("Default");
+		if (!SkinWeights)
+		{
+			return false;
+		}
+
+		// store ids as quads
+		TArray<TArray<int>> FoundQuads;
+		if (InParams.bTryUsingQuads)
+		{
+			// do not copy on the mesh now
+			TGuardValue<bool> CopyGuard(bCopyToMesh, false);
+			if (FindSourceMeshPolygonPolygroups(
+				InParams.bRespectUVSeams,
+				InParams.bRespectNormalSeams,
+				InParams.QuadAdjacencyWeight,
+				InParams.QuadMetricClamp,
+				InParams.MaxSearchRounds))
+			{
+				::Swap(FoundPolygroups, FoundQuads);
+			}
+		}
+
+		static constexpr float Third = 1.f / 3.f;
+	    static constexpr FBoneWeightsSettings BlendSettings;
+		
+		TArray<FBoneWeights> TriangleWeights;
+		TriangleWeights.Reserve(Mesh->TriangleCount());
+
+		// store per triangle bone weights
+		const int32 NumTriangles = Mesh->TriangleCount();
+		for (int32 Index = 0; Index < NumTriangles; Index++)
+		{
+			const FIndex3i& Tri = Mesh->GetTriangleRef(Index);
+
+			FBoneWeights Weights0, Weights1, Weights2;
+			SkinWeights->GetValue(Tri[0], Weights0);
+			SkinWeights->GetValue(Tri[1], Weights1);
+			SkinWeights->GetValue(Tri[2], Weights2);
+			
+			TriangleWeights.Add( FBoneWeights::Blend(Weights0, Weights1, Weights2, Third, Third, Third, BlendSettings) );
+		}
+
+		TMap<FBoneIndexType, TArray<int32>> TrianglesPerBone;
+
+		TArray<FBoneIndexType> TriangleBone;
+		TriangleBone.Init(0, NumTriangles);
+
+		auto MoveToParentIndexIfRemoved = [&TriangleBone, &InParams](int32 TriIndex)
+		{
+			FBoneIndexType& BoneIndex = TriangleBone[TriIndex];
+			if (BoneIndex != 0)
+			{
+				while (InParams.BonesToRemove.Contains(BoneIndex) && BoneIndex > 0)
+				{
+					if (ensure(InParams.Parents.IsValidIndex(BoneIndex)))
+					{
+						BoneIndex = FMath::Max(0, InParams.Parents[BoneIndex]);
+					}
+					else
+					{
+						BoneIndex = 0;
+					}
+				}
+			}
+		};
+		
+		if (FoundQuads.IsEmpty())
+		{
+			// store per bone triangles using all the triangles
+			for (int32 TriIndex = 0; TriIndex < NumTriangles; TriIndex++)
+			{
+				const FBoneWeights& Weights = TriangleWeights[TriIndex];
+				if (Weights.Num() > 0)
+				{
+					TriangleBone[TriIndex] = Weights[0].GetBoneIndex();
+					MoveToParentIndexIfRemoved(TriIndex);
+				}
+			}
+		}
+		else
+		{
+			// store per bone quads
+			for (const TArray<int>& QuadTris: FoundQuads)
+			{
+				if (QuadTris.Num() == 2)
+				{
+					const int32 TriIndex0 = QuadTris[0], TriIndex1 = QuadTris[1];
+					const FBoneWeights& Weights0 = TriangleWeights[TriIndex0];
+					const FBoneWeights& Weights1 = TriangleWeights[TriIndex1];
+
+					const FBoneWeights& QuadWeights = FBoneWeights::Blend(Weights0, Weights1, 0.5f, BlendSettings);
+					
+					if (QuadWeights.Num() > 0)
+					{
+						const FBoneWeight BoneWeight = QuadWeights[0];
+						TriangleBone[TriIndex0] = BoneWeight.GetBoneIndex();
+						MoveToParentIndexIfRemoved(TriIndex0);
+						TriangleBone[TriIndex1] = BoneWeight.GetBoneIndex();
+						MoveToParentIndexIfRemoved(TriIndex1);
+					}
+				}
+				else
+				{
+					for (const int32 TriIndex: QuadTris)
+					{
+						const FBoneWeights& Weights = TriangleWeights[TriIndex];
+						if (Weights.Num() > 0)
+						{
+							TriangleBone[TriIndex] = Weights[0].GetBoneIndex();
+							MoveToParentIndexIfRemoved(TriIndex);
+						}
+					}
+				}
+			}
+		}
+
+		for (int32 TriIndex = 0; TriIndex < NumTriangles; TriIndex++)
+		{
+			TArray<int32>& BoneTriangles = TrianglesPerBone.FindOrAdd(TriangleBone[TriIndex]);
+			BoneTriangles.Add(TriIndex);
+		}
+		
+		for (const auto&[BoneIndex, Triangles]: TrianglesPerBone)
+		{
+			FoundPolygroups.Add(Triangles);
+		}
+
+		if (bApplyPostProcessing)
+		{
+			// optimize to the parent ?
+		}
+
+		if (bCopyToMesh)
+		{
+			Mesh->EnableTriangleGroups(0);
+			for (const auto&[BoneIndex, Triangles]: TrianglesPerBone)
+			{
+				for (const int32 TriID: Triangles)
+				{
+					Mesh->SetTriangleGroup(TriID, BoneIndex);
+				}
+			}
+		}
+		
+		return !FoundPolygroups.IsEmpty();
+	}
+};
+
+class FDirectMeshPolygroupOp : public  FDynamicMeshOperator
+{
+public:
+	// parameters set by the tool
+	int32 MinGroupSize = 2;
+	int32 InitialGroupID = 0;
+	FDirectMeshPolygroupsGenerator::FPolygroupsFromSkinWeightsParams Params;
+
+	// input mesh
+	TSharedPtr<FSharedConstDynamicMesh3, ESPMode::ThreadSafe> OriginalMesh;
+	// input polygroups, if available
+	TSharedPtr<FPolygroupSet, ESPMode::ThreadSafe> SourcePolyGroups;
+
+	// result
+	FDirectMeshPolygroupsGenerator Generator;
+
+	virtual void CalculateResult(FProgressCancel* Progress) override
+	{
+		if ((Progress && Progress->Cancelled()) || OriginalMesh.IsValid() == false)
+		{
+			return;
+		}
+
+		OriginalMesh->AccessSharedObject([&](const FDynamicMesh3& ReadMesh)
+		{
+			ResultMesh->Copy(ReadMesh, true, true, true, true);
+		});
+	
+		if (Progress && Progress->Cancelled())
+		{
+			return;
+		}
+
+		Generator = FDirectMeshPolygroupsGenerator(ResultMesh.Get());
+		Generator.MinGroupSize = MinGroupSize;
+		Generator.InitialGroupID = InitialGroupID;
+		
+		Generator.FindPolygroupsFromSkinWeights(Params);
+
+		Generator.FindPolygroupEdges();
+	}
+
+	void SetTransform(const FTransformSRT3d& Transform)
+	{
+		ResultTransform = Transform;
+	}
+
+
+};
+
+namespace UE::Geometry
+{
+
+FDynamicMeshTriangleLabelAttribute* GetTriangleLabelAttribute(FDynamicMesh3& InMesh, const FName LayerName)
+{
+	checkSlow(InMesh.HasAttributes());
+
+	FDynamicMeshAttributeSet* Attributes = InMesh.Attributes();
+	FDynamicMeshTriangleLabelAttribute* Attr = Attributes->FindTriangleLabelAttribute(LayerName);
+	if (!Attr)
+	{
+		Attr = new FDynamicMeshTriangleLabelAttribute(&InMesh);
+		Attributes->AttachTriangleLabelAttribute(LayerName, Attr);
+	}
+
+	return Attr;
+}
+	
+}
+
+TUniquePtr<FDynamicMeshOperator> UDirectMeshPolygroupOperatorFactory::MakeNewOperator()
+{
+	// backpointer used to populate parameters.
+	check(DirectMeshPolygroupTool);
+
+	// Create the actual operator type based on the requested operation
+	TUniquePtr<FDirectMeshPolygroupOp>  MeshOp = MakeUnique<FDirectMeshPolygroupOp>();
+
+	// Operator runs on another thread - copy data over that it needs.
+	DirectMeshPolygroupTool->UpdateOpParameters(*MeshOp);
+
+	// give the operator
+	return MeshOp;
+}
+
+/*
+ * Tool
+ */
+UDirectMeshPolygroupTool::UDirectMeshPolygroupTool()
+{}
+
+bool UDirectMeshPolygroupTool::CanAccept() const
+{
+	return Super::CanAccept() && (PreviewCompute == nullptr || PreviewCompute->HaveValidResult());
+}
+
+void UDirectMeshPolygroupTool::Setup()
+{
+	UInteractiveTool::Setup();
+	SetToolDisplayName(LOCTEXT("DirectMeshPolygroupToolName", "Generate PolyGroups"));
+	
+	FComponentMaterialSet MaterialSet = UE::ToolTarget::GetMaterialSet(Target);
+
+	OriginalDynamicMesh = MakeShared<FDynamicMesh3, ESPMode::ThreadSafe>(UE::ToolTarget::GetDynamicMeshCopy(Target));
+
+	// initialize triangle ROI if one exists
+	SelectionTriangleROI = MakeShared<TSet<int32>, ESPMode::ThreadSafe>();
+	TArray<int> TriangleROI;
+	if (HasGeometrySelection())
+	{
+		const FGeometrySelection& InputSelection = GetGeometrySelection();
+
+		EnumerateSelectionTriangles(InputSelection, *OriginalDynamicMesh, [&](int32 TriangleID)
+		{
+			SelectionTriangleROI->Add(TriangleID);
+		});
+
+		TriangleROI = SelectionTriangleROI->Array();
+
+		OriginalSubmesh = MakeShared<FDynamicSubmesh3, ESPMode::ThreadSafe>(OriginalDynamicMesh.Get(), TriangleROI);
+		bUsingSelection = true;
+	}
+
+	if (bUsingSelection)
+	{
+		ComputeOperatorSharedMesh = MakeShared<FSharedConstDynamicMesh3>(&OriginalSubmesh->GetSubmesh());
+	}
+	else
+	{
+		ComputeOperatorSharedMesh = MakeShared<FSharedConstDynamicMesh3>(OriginalDynamicMesh.Get());
+	}
+
+	Properties = NewObject<UDirectMeshPolygroupToolProperties>(this);
+	Properties->RestoreProperties(this);
+	AddToolPropertySource(Properties);
+	FTransform MeshTransform = (FTransform)UE::ToolTarget::GetLocalToWorldTransform(Target);
+	UE::ToolTarget::HideSourceObject(Target);
+
+	{
+		// create the operator factory
+		UDirectMeshPolygroupOperatorFactory* DirectMeshPolygroupOperatorFactory = NewObject<UDirectMeshPolygroupOperatorFactory>(this);
+		DirectMeshPolygroupOperatorFactory->DirectMeshPolygroupTool = this; // set the back pointer
+
+		PreviewCompute = NewObject<UMeshOpPreviewWithBackgroundCompute>(DirectMeshPolygroupOperatorFactory);
+		PreviewCompute->Setup(GetTargetWorld(), DirectMeshPolygroupOperatorFactory);
+		ToolSetupUtil::ApplyRenderingConfigurationToPreview(PreviewCompute->PreviewMesh, Target);
+		PreviewCompute->SetIsMeshTopologyConstant(true, EMeshRenderAttributeFlags::Positions | EMeshRenderAttributeFlags::VertexNormals);
+
+		// Give the preview something to display
+		PreviewCompute->PreviewMesh->SetTransform(MeshTransform);
+		PreviewCompute->PreviewMesh->SetTangentsMode(EDynamicMeshComponentTangentsMode::AutoCalculated);
+		PreviewCompute->PreviewMesh->UpdatePreview(OriginalDynamicMesh.Get());
+		
+		PreviewCompute->ConfigureMaterials(MaterialSet.Materials,
+			ToolSetupUtil::GetDefaultWorkingMaterial(GetToolManager())
+		);
+
+		// show the preview mesh
+		PreviewCompute->SetVisibility(true);
+
+		// something to capture the polygons from the async task when it is done
+		PreviewCompute->OnOpCompleted.AddLambda([this](const FDynamicMeshOperator* MeshOp) 
+		{ 
+			const FDirectMeshPolygroupOp*  DirectMeshPolygroupOp = static_cast<const FDirectMeshPolygroupOp*>(MeshOp);
+			this->PolygonEdges = DirectMeshPolygroupOp->Generator.PolygroupEdges;
+			UpdateVisualization();
+		});
+
+	}
+	
+	if (bUsingSelection)
+	{
+		// create the preview object for the unmodified area
+		UnmodifiedAreaPreviewMesh = NewObject<UPreviewMesh>();
+		UnmodifiedAreaPreviewMesh->CreateInWorld(GetTargetWorld(), MeshTransform);
+		ToolSetupUtil::ApplyRenderingConfigurationToPreview(UnmodifiedAreaPreviewMesh, Target);
+		UnmodifiedAreaPreviewMesh->SetMaterials(MaterialSet.Materials);
+		UnmodifiedAreaPreviewMesh->EnableSecondaryTriangleBuffers( [this](const FDynamicMesh3* Mesh, int32 TriangleID) {
+			return SelectionTriangleROI->Contains(TriangleID);
+		});
+		UnmodifiedAreaPreviewMesh->SetSecondaryBuffersVisibility(false);
+		UnmodifiedAreaPreviewMesh->SetTangentsMode(EDynamicMeshComponentTangentsMode::AutoCalculated);
+		UnmodifiedAreaPreviewMesh->UpdatePreview(OriginalDynamicMesh.Get());
+	}
+
+
+	PreviewGeometry = NewObject<UPreviewGeometry>(this);
+	PreviewGeometry->CreateInWorld(GetTargetWorld(), MeshTransform);
+
+	Properties->WatchProperty(Properties->bShowGroupColors, [this](bool) { UpdateVisualization(); });
+	Properties->WatchProperty(Properties->MinGroupSize, [this](int32) { OnSettingsModified(); });
+	Properties->WatchProperty(Properties->QuadAdjacencyWeight, [this](float) { OnSettingsModified(); });
+	Properties->WatchProperty(Properties->QuadMetricClamp, [this](float) { OnSettingsModified(); });
+	Properties->WatchProperty(Properties->QuadSearchRounds, [this](int) { OnSettingsModified(); });
+	Properties->WatchProperty(Properties->bRespectUVSeams, [this](int) { OnSettingsModified(); });
+	Properties->WatchProperty(Properties->bRespectHardNormals, [this](int) { OnSettingsModified(); });
+	Properties->WatchProperty(Properties->bTryUsingQuads, [this](bool) { OnSettingsModified(); });
+	Properties->WatchProperty(Properties->BonesToRemove, [this](TArray<FName>){ OnSettingsModified(); });
+
+	if (bUsingSelection)
+	{
+		GeometrySelectionVizProperties = NewObject<UGeometrySelectionVisualizationProperties>(this);
+		GeometrySelectionVizProperties->RestoreProperties(this);
+		AddToolPropertySource(GeometrySelectionVizProperties);
+		GeometrySelectionVizProperties->Initialize(this);
+		GeometrySelectionVizProperties->bEnableShowTriangleROIBorder = true;
+		GeometrySelectionVizProperties->SelectionElementType = static_cast<EGeometrySelectionElementType>(GeometrySelection.ElementType);
+		GeometrySelectionVizProperties->SelectionTopologyType = static_cast<EGeometrySelectionTopologyType>(GeometrySelection.TopologyType);
+
+		GeometrySelectionViz = NewObject<UPreviewGeometry>(this);
+		GeometrySelectionViz->CreateInWorld(GetTargetWorld(), MeshTransform);
+		
+		InitializeGeometrySelectionVisualization(
+			GeometrySelectionViz,
+			GeometrySelectionVizProperties,
+			*OriginalDynamicMesh,
+			GeometrySelection,
+			nullptr,
+			nullptr,
+			&TriangleROI);
+	}
+
+	// start the compute
+	PreviewCompute->InvalidateResult();
+
+	// updates the triangle color visualization
+	UpdateVisualization();
+
+	GetToolManager()->DisplayMessage(
+		LOCTEXT("OnStartTool", "Cluster triangles of the Mesh into PolyGroups using various strategies"),
+		EToolMessageLevel::UserNotification);
+}
+
+void UDirectMeshPolygroupTool::UpdateOpParameters(FDirectMeshPolygroupOp& DirectMeshPolygroupOp) const
+{
+	DirectMeshPolygroupOp.MinGroupSize = Properties->MinGroupSize;
+	DirectMeshPolygroupOp.Params.QuadMetricClamp = Properties->QuadMetricClamp;
+	DirectMeshPolygroupOp.Params.QuadAdjacencyWeight = Properties->QuadAdjacencyWeight;
+	DirectMeshPolygroupOp.Params.MaxSearchRounds = FMath::Clamp(Properties->QuadSearchRounds, 1, 100);
+	DirectMeshPolygroupOp.Params.bRespectUVSeams = Properties->bRespectUVSeams;
+	DirectMeshPolygroupOp.Params.bRespectNormalSeams = Properties->bRespectHardNormals;
+	DirectMeshPolygroupOp.Params.bTryUsingQuads = Properties->bTryUsingQuads;
+
+	DirectMeshPolygroupOp.OriginalMesh = ComputeOperatorSharedMesh;
+
+	DirectMeshPolygroupOp.Params.BonesToRemove.Reset();
+	DirectMeshPolygroupOp.Params.Parents.Reset();
+	
+	TArray<FName> BoneNames;
+	ComputeOperatorSharedMesh->AccessSharedObject([&BoneNames, &DirectMeshPolygroupOp](const FDynamicMesh3& ReadMesh)
+	{
+		if (const FDynamicMeshAttributeSet* Attributes = ReadMesh.Attributes())
+		{
+			if (Attributes->HasBones())
+			{
+				BoneNames = Attributes->GetBoneNames()->GetAttribValues();
+				DirectMeshPolygroupOp.Params.Parents = Attributes->GetBoneParentIndices()->GetAttribValues(); 
+			}
+		}
+	});
+
+	ensure(!BoneNames.IsEmpty());
+	
+	for (int32 BoneIndex = 0; BoneIndex < BoneNames.Num(); BoneIndex++)
+	{
+		const FString BoneNameStr = BoneNames[BoneIndex].ToString();
+		const bool bRemove = Properties->BonesToRemove.ContainsByPredicate([BoneNameStr](const FName& BoneToRemove)
+		{
+			return BoneNameStr.Contains(BoneToRemove.ToString());
+		});
+		if (bRemove)
+		{
+			DirectMeshPolygroupOp.Params.BonesToRemove.Add(BoneIndex);
+		}
+	}
+
+	if (bUsingSelection)
+	{
+		DirectMeshPolygroupOp.InitialGroupID = OriginalDynamicMesh->MaxGroupID();
+	}
+
+	DirectMeshPolygroupOp.SetTransform(UE::ToolTarget::GetLocalToWorldTransform(Target));
+}
+
+void UDirectMeshPolygroupTool::OnShutdown(EToolShutdownType ShutdownType)
+{
+	Properties->SaveProperties(this);
+	UE::ToolTarget::ShowSourceObject(Target);
+
+	if (PreviewGeometry)
+	{
+		PreviewGeometry->Disconnect();
+		PreviewGeometry = nullptr;
+	}
+
+	if (UnmodifiedAreaPreviewMesh != nullptr)
+	{
+		UnmodifiedAreaPreviewMesh->Disconnect();
+		UnmodifiedAreaPreviewMesh = nullptr;
+	}
+
+	if (PreviewCompute)
+	{
+		FDynamicMeshOpResult Result = PreviewCompute->Shutdown();
+		if (ShutdownType == EToolShutdownType::Accept)
+		{
+			GetToolManager()->BeginUndoTransaction(LOCTEXT("DirectMeshPolygroupToolTransactionName", "Save DMC Layer"));
+
+			FDynamicMesh3* DynamicMeshResult = Result.Mesh.Get();
+			if (ensure(DynamicMeshResult != nullptr))
+			{
+				FDynamicMesh3 UseResultMesh = *OriginalDynamicMesh;
+				if (UseResultMesh.HasAttributes() == false)
+				{
+					UseResultMesh.EnableAttributes();
+				}
+				
+				// if we want to write to any layer other than default, we have to find or create it
+				FDynamicMeshTriangleLabelAttribute* LabelAttribute = GetTriangleLabelAttribute(UseResultMesh, Properties->LayerName);
+				if (ensure(LabelAttribute))
+				{
+					const FDynamicMeshAttributeSet* Attributes = UseResultMesh.Attributes();
+					const TArray<FName> BoneNames =
+						Attributes && Attributes->HasBones() ? Attributes->GetBoneNames()->GetAttribValues() : TArray<FName>();
+
+					for (const int32 TriangleID : UseResultMesh.TriangleIndicesItr())
+					{
+						const int32 BoneIndex = DynamicMeshResult->GetTriangleGroup(TriangleID);
+						LabelAttribute->SetValue(TriangleID, BoneNames.IsValidIndex(BoneIndex) ? BoneNames[BoneIndex] : NAME_None);
+					}
+					
+					UE::ToolTarget::CommitDynamicMeshUpdate(Target, UseResultMesh, true);
+				}
+				else
+				{
+					// If we can't find or create the layer (which should not be possible) the tool is going to do nothing,
+					// this is the safest option at this point
+					UE_LOG(LogGeometry, Warning, TEXT("Output group layer missing?"));
+				}
+			}
+
+			if (bUsingSelection)
+			{
+				// if the input was a group selection, that selection is no longer valid. But if we output
+				// a triangle selection it should be converted to the group selection
+				FGeometrySelection OutputSelection;
+				for (int32 tid : *SelectionTriangleROI)
+				{
+					OutputSelection.Selection.Add(FGeoSelectionID::MeshTriangle(tid).Encoded());
+				}
+				SetToolOutputGeometrySelectionForTarget(this, Target, OutputSelection);
+			}
+
+			GetToolManager()->EndUndoTransaction();
+		}
+	}
+
+	if (ComputeOperatorSharedMesh.IsValid())
+	{
+		ComputeOperatorSharedMesh->ReleaseSharedObject();
+	}
+
+	Super::OnShutdown(ShutdownType);
+}
+
+void UDirectMeshPolygroupTool::OnSettingsModified()
+{
+	PreviewCompute->InvalidateResult();
+}
+
+
+void UDirectMeshPolygroupTool::OnTick(float DeltaTime)
+{
+	Super::OnTick(DeltaTime);
+
+	PreviewCompute->Tick(DeltaTime);
+}
+
+void UDirectMeshPolygroupTool::UpdateVisualization()
+{
+	if (!PreviewCompute)
+	{
+		return;
+	}
+
+	IMaterialProvider* MaterialTarget = Cast<IMaterialProvider>(Target);
+	FComponentMaterialSet MaterialSet;
+	
+	if (Properties->bShowGroupColors)
+	{
+		UMaterialInstanceDynamic* VtxColorMaterial = ToolSetupUtil::GetVertexColorMaterial(GetToolManager());
+		MaterialSet.Materials.Init(VtxColorMaterial, MaterialTarget->GetNumMaterials());
+		
+		PreviewCompute->PreviewMesh->SetTriangleColorFunction([this](const FDynamicMesh3* Mesh, int TriangleID)
+		{
+			return LinearColors::SelectFColor(Mesh->GetTriangleGroup(TriangleID));
+		}, 
+		UPreviewMesh::ERenderUpdateMode::FastUpdate);
+	}
+	else
+	{
+		MaterialTarget->GetMaterialSet(MaterialSet);
+		PreviewCompute->PreviewMesh->ClearTriangleColorFunction(UPreviewMesh::ERenderUpdateMode::FastUpdate);
+	}
+	
+	PreviewCompute->ConfigureMaterials(MaterialSet.Materials, ToolSetupUtil::GetDefaultWorkingMaterial(GetToolManager()));
+
+	FColor GroupLineColor = FColor::Red;
+	float GroupLineThickness = 2.0f;
+
+	ComputeOperatorSharedMesh->AccessSharedObject([&](const FDynamicMesh3& ReadMesh)
+	{
+		PreviewGeometry->CreateOrUpdateLineSet(TEXT("GroupBorders"), PolygonEdges.Num(), 
+			[&](int32 k, TArray<FRenderableLine>& LinesOut) {
+				FVector3d A, B;
+				ReadMesh.GetEdgeV(PolygonEdges[k], A, B);
+				LinesOut.Add(FRenderableLine(A, B, GroupLineColor, GroupLineThickness));
+			}, 1);
+	});
+}
+
+
+
+#undef LOCTEXT_NAMESPACE
+
diff --git a/Engine/Plugins/Experimental/Animation/DirectMeshControl/Source/DirectMeshControl/Public/DirectMeshControlCommands.h b/Engine/Plugins/Experimental/Animation/DirectMeshControl/Source/DirectMeshControl/Public/DirectMeshControlCommands.h
new file mode 100644
--- /dev/null
+++ b/Engine/Plugins/Experimental/Animation/DirectMeshControl/Source/DirectMeshControl/Public/DirectMeshControlCommands.h
@@ -0,0 +1,23 @@
+﻿// Copyright Epic Games, Inc. All Rights Reserved.
+
+#pragma once
+
+#include "CoreMinimal.h"
+#include "Framework/Commands/Commands.h"
+
+/**
+ * 
+ */
+class DIRECTMESHCONTROL_API FDirectMeshControlCommands : public TCommands<FDirectMeshControlCommands>
+{
+public:
+	FDirectMeshControlCommands();
+
+	TSharedPtr<FUICommandInfo> BeginDirectMeshPolygroupTool;
+	TSharedPtr<FUICommandInfo> BeginDirectMeshControlTool;
+
+	/**
+	 * Initialize commands
+	 */
+	virtual void RegisterCommands() override;
+};
\ No newline at end of file
diff --git a/Engine/Plugins/Experimental/Animation/DirectMeshControl/Source/DirectMeshControl/Public/DirectMeshControlComponent.h b/Engine/Plugins/Experimental/Animation/DirectMeshControl/Source/DirectMeshControl/Public/DirectMeshControlComponent.h
new file mode 100644
--- /dev/null
+++ b/Engine/Plugins/Experimental/Animation/DirectMeshControl/Source/DirectMeshControl/Public/DirectMeshControlComponent.h
@@ -0,0 +1,23 @@
+// Copyright Epic Games, Inc. All Rights Reserved.
+
+#pragma once
+
+#include "Components/SkeletalMeshComponent.h"
+
+#include "DirectMeshControlComponent.generated.h"
+
+#define UE_API DIRECTMESHCONTROL_API
+
+/**
+ * 
+ */
+UCLASS(transient)
+class UE_API UDirectMeshControlComponent : public USkeletalMeshComponent
+{
+	GENERATED_BODY()
+ 	
+	virtual void CreateRenderState_Concurrent(FRegisterComponentContext* Context) override;
+	virtual void SendRenderDynamicData_Concurrent() override;
+};
+
+#undef UE_API
\ No newline at end of file
diff --git a/Engine/Plugins/Experimental/Animation/DirectMeshControl/Source/DirectMeshControl/Public/DirectMeshControlModule.h b/Engine/Plugins/Experimental/Animation/DirectMeshControl/Source/DirectMeshControl/Public/DirectMeshControlModule.h
new file mode 100644
--- /dev/null
+++ b/Engine/Plugins/Experimental/Animation/DirectMeshControl/Source/DirectMeshControl/Public/DirectMeshControlModule.h
@@ -0,0 +1,30 @@
+﻿// Copyright Epic Games, Inc. All Rights Reserved.
+
+#pragma once
+
+#include "Modules/ModuleManager.h"
+#include "SkeletalMeshModelingModeToolExtensions.h"
+
+struct FExtensionToolDescription;
+struct FExtensionToolQueryInfo;
+class IControlRigEditor;
+class FUICommandList;
+class FExtender;
+class FRigVMEditorBase;
+
+/**
+ * 
+ */
+class FDirectMeshControlModule : public IModuleInterface, public ISkeletalMeshModelingModeToolExtension
+{
+public:
+
+	/** IModuleInterface implementation */
+	virtual void StartupModule() override;
+	virtual void ShutdownModule() override;
+
+	/** IModelingModeToolExtension implementation */
+	virtual FText GetExtensionName() override;
+	virtual FText GetToolSectionName() override;
+	virtual void GetExtensionTools(const FExtensionToolQueryInfo& InQueryInfo, TArray<FExtensionToolDescription>& OutTools) override;
+};
\ No newline at end of file
diff --git a/Engine/Plugins/Experimental/Animation/DirectMeshControl/Source/DirectMeshControl/Public/DirectMeshControlSettings.h b/Engine/Plugins/Experimental/Animation/DirectMeshControl/Source/DirectMeshControl/Public/DirectMeshControlSettings.h
new file mode 100644
--- /dev/null
+++ b/Engine/Plugins/Experimental/Animation/DirectMeshControl/Source/DirectMeshControl/Public/DirectMeshControlSettings.h
@@ -0,0 +1,28 @@
+// Copyright Epic Games, Inc. All Rights Reserved.
+
+#pragma once
+
+#include "Engine/DeveloperSettings.h"
+
+#include "DirectMeshControlSettings.generated.h"
+
+#define UE_API DIRECTMESHCONTROL_API
+
+UCLASS(MinimalAPI, config=EditorPerProjectUserSettings, meta = (Experimental, DisplayName = "Direct Mesh Control Settings"))
+class UDirectMeshControlSettings : public UDeveloperSettings
+{
+	GENERATED_UCLASS_BODY()
+
+public:
+
+	/** The overlay color used when hovering a direct mesh */
+	UPROPERTY(config, EditAnywhere, Category = "Direct Mesh Control Settings")
+	FLinearColor OverlayColor = FLinearColor(0.25f, 0.25f, 0.25f, 0.25f);
+	
+	DECLARE_MULTICAST_DELEGATE_OneParam(FOnUpdateSettings,  const UDirectMeshControlSettings*);
+	static UE_API FOnUpdateSettings OnSettingsChange;
+
+	UE_API virtual void PostEditChangeProperty(FPropertyChangedEvent& PropertyChangedEvent) override;
+};
+
+#undef UE_API
diff --git a/Engine/Plugins/Experimental/Animation/DirectMeshControl/Source/DirectMeshControl/Public/DirectMeshControlStyle.h b/Engine/Plugins/Experimental/Animation/DirectMeshControl/Source/DirectMeshControl/Public/DirectMeshControlStyle.h
new file mode 100644
--- /dev/null
+++ b/Engine/Plugins/Experimental/Animation/DirectMeshControl/Source/DirectMeshControl/Public/DirectMeshControlStyle.h
@@ -0,0 +1,20 @@
+﻿// Copyright Epic Games, Inc. All Rights Reserved.
+
+#pragma once
+
+#include "CoreMinimal.h"
+#include "Styling/ISlateStyle.h"
+#include "Styling/SlateStyle.h"
+
+class DIRECTMESHCONTROL_API FDirectMeshControlStyle
+{
+public:
+	static void Initialize();
+	static void Shutdown();
+
+	static TSharedPtr<class ISlateStyle> Get();
+	static FName GetStyleSetName();
+
+private:
+	static TSharedPtr<class FSlateStyleSet> StyleSet;
+};
\ No newline at end of file
diff --git a/Engine/Plugins/Experimental/Animation/DirectMeshControl/Source/DirectMeshControl/Public/DirectMeshControlUtilities.h b/Engine/Plugins/Experimental/Animation/DirectMeshControl/Source/DirectMeshControl/Public/DirectMeshControlUtilities.h
new file mode 100644
--- /dev/null
+++ b/Engine/Plugins/Experimental/Animation/DirectMeshControl/Source/DirectMeshControl/Public/DirectMeshControlUtilities.h
@@ -0,0 +1,97 @@
+// Copyright Epic Games, Inc. All Rights Reserved.
+
+#pragma once
+
+#include "Engine/SkeletalMesh.h"
+#include "UObject/SoftObjectPtr.h"
+
+#include "DirectMeshControlUtilities.generated.h"
+
+#define UE_API DIRECTMESHCONTROL_API
+
+class UMaterialInstanceDynamic;
+class UOptimusDeformer;
+class UMaterial;
+
+namespace UE::Geometry
+{
+	class FDynamicMesh3;
+}
+
+/**
+ * FGroupSubMeshes
+ */
+
+USTRUCT()
+struct UE_API FGroupSubMeshes
+{
+	GENERATED_BODY()
+	
+	void Rebuild(USkeletalMesh* InSkeletalMesh, const UE::Geometry::FDynamicMesh3* Mesh, const FName LayerName);
+	void Reset();
+
+	const TArray<TObjectPtr<USkeletalMesh>>& GetSubSkeletalMeshes() const;
+	const TMap<int32, int32>& GetGroupIdToSubMesh() const;
+	
+private:
+	
+	UPROPERTY()
+	TArray<TObjectPtr<USkeletalMesh>> SubSkeletalMeshes;
+	UPROPERTY()
+	TMap<int32, int32> GroupIdToSubMesh;
+};
+
+/**
+ * FSubSkeletalMeshData
+ */
+
+USTRUCT()
+struct FSubSkeletalMeshData
+{
+	GENERATED_BODY()
+	
+	FSubSkeletalMeshData() = default;
+	
+	FSubSkeletalMeshData(USkeletalMesh* InSkeletalMesh, const FName InLayerName)
+		: SkeletalMesh(InSkeletalMesh)
+		, LayerName(InLayerName)
+	{}
+
+	UPROPERTY()
+	TSoftObjectPtr<USkeletalMesh> SkeletalMesh = nullptr;
+	
+	UPROPERTY()
+	FName LayerName = NAME_None;
+	
+	UPROPERTY()
+	uint32 Hash = 0;
+	
+	UPROPERTY()
+	FGroupSubMeshes GroupSubMeshes;
+
+	friend FORCEINLINE uint32 GetTypeHash(const FSubSkeletalMeshData& Data)
+	{
+		return HashCombine(GetTypeHash(Data.SkeletalMesh), GetTypeHash(Data.LayerName));
+	}
+
+	bool operator ==(const FSubSkeletalMeshData& Other) const
+	{
+		return SkeletalMesh == Other.SkeletalMesh && LayerName == Other.LayerName;
+	}
+};
+
+/**
+ * 
+ */
+
+namespace UE::DMC
+{
+	UE_API bool UseMaterialHover();
+	UE_API const FGroupSubMeshes& GetSubMeshes(USkeletalMesh* InSkeletalMesh, const Geometry::FDynamicMesh3* Mesh, const FName LayerName);
+	UE_API const FName& GetSubToSourceAttrName();
+	UE_API const FName& GetColorVarName();
+	UE_API UMaterial* GetMaterial();
+	UE_API UOptimusDeformer* GetDeformer();
+}
+	
+#undef UE_API
diff --git a/Engine/Plugins/Experimental/Animation/DirectMeshControl/Source/DirectMeshControl/Public/Tools/DirectMeshControlTool.h b/Engine/Plugins/Experimental/Animation/DirectMeshControl/Source/DirectMeshControl/Public/Tools/DirectMeshControlTool.h
new file mode 100644
--- /dev/null
+++ b/Engine/Plugins/Experimental/Animation/DirectMeshControl/Source/DirectMeshControl/Public/Tools/DirectMeshControlTool.h
@@ -0,0 +1,233 @@
+// Copyright Epic Games, Inc. All Rights Reserved.
+
+#pragma once
+
+#include "GroupTopology.h"
+#include "InteractiveToolBuilder.h"
+#include "ReferenceSkeleton.h"
+#include "SingleSelectionTool.h"
+#include "SkeletalMesh/SkeletalMeshEditingInterface.h"
+#include "BaseBehaviors/BehaviorTargetInterfaces.h"
+#include "DynamicMesh/DynamicMeshAABBTree3.h"
+#include "Polygroups/PolygroupSet.h"
+#include "DynamicMesh/DynamicMesh3.h"
+
+#include "DirectMeshControlTool.generated.h"
+
+namespace UE::Geometry
+{
+	class ITriangleGroupSetAdapter;
+}
+
+class UPreviewMesh;
+class USingleClickInputBehavior;
+class USkeletalMesh;
+class UDynamicMeshComponent;
+class UPolygonSelectionMechanic;
+class UPersonaEditorModeManagerContext;
+class UOptimusDeformer;
+using UE::Geometry::FDynamicMesh3;
+
+/**
+ * UDirectMeshControlToolBuilder
+ */
+
+UCLASS()
+class DIRECTMESHCONTROL_API UDirectMeshControlToolBuilder : public UInteractiveToolWithToolTargetsBuilder
+{
+	GENERATED_BODY()
+
+public:
+	// UInteractiveToolBuilder overrides
+	virtual bool CanBuildTool(const FToolBuilderState& SceneState) const override;
+	virtual UInteractiveTool* BuildTool(const FToolBuilderState& SceneState) const override;
+
+	TFunction<void (const USkeletalMeshComponent*, const FName)> PreClickFunction;
+	
+protected:
+	// UInteractiveToolWithToolTargetsBuilder overrides
+	virtual const FToolTargetTypeRequirements& GetTargetRequirements() const override;
+};
+
+/**
+ * UDirectMeshControlTool
+ */
+
+UCLASS()
+class DIRECTMESHCONTROL_API UDirectMeshControlTool :
+	public USingleSelectionTool,
+	public IClickDragBehaviorTarget,
+	public IHoverBehaviorTarget,
+	public ISkeletalMeshEditingInterface
+{
+	GENERATED_BODY()
+
+public:
+
+	void Init(const FToolBuilderState& InSceneState);
+	
+	// UInteractiveTool overrides
+	virtual void Setup() override;
+	virtual void Shutdown(EToolShutdownType ShutdownType) override;
+
+	virtual void DrawHUD(FCanvas* Canvas, IToolsContextRenderAPI* RenderAPI) override;
+	virtual void Render(IToolsContextRenderAPI* RenderAPI) override;
+	
+	virtual bool HasCancel() const override { return true; }
+	virtual bool HasAccept() const override { return true; }
+	virtual bool CanAccept() const override { return true; }
+
+	virtual void OnPropertyModified(UObject* PropertySet, FProperty* Property) override;
+
+	// ICLickDragBehaviorTarget overrides
+	virtual FInputRayHit CanBeginClickDragSequence(const FInputDeviceRay& InPressPos) override;
+	
+	virtual void OnClickPress(const FInputDeviceRay& InPressPos) override;
+	virtual void OnClickDrag(const FInputDeviceRay& InDragPos) override;
+	virtual void OnClickRelease(const FInputDeviceRay& InReleasePos) override;
+	virtual void OnTerminateDragSequence() override;
+
+	virtual FInputRayHit BeginHoverSequenceHitTest(const FInputDeviceRay& PressPos) override;
+	virtual void OnBeginHover(const FInputDeviceRay& DevicePos) override;
+	virtual bool OnUpdateHover(const FInputDeviceRay& DevicePos) override;
+	virtual void OnEndHover() override;
+	
+	virtual void OnTick(float DeltaTime) override;
+
+	// IInteractiveToolCameraFocusAPI overrides
+	virtual FBox GetWorldSpaceFocusBox() override;
+
+	bool bRebuild = true;
+	
+protected:
+
+	// UInteractiveTool overrides
+	virtual void RegisterActions(FInteractiveToolActionSet& ActionSet) override;
+
+	// ISkeletalMeshEditionInterface overrides
+	virtual void HandleSkeletalMeshModified(const TArray<FName>& BoneNames, const ESkeletalMeshNotifyType InNotifyType) override;
+
+public:
+	
+	UPROPERTY()
+	TObjectPtr<UDirectMeshControlProperties> Properties;
+
+	UPROPERTY()
+	TObjectPtr<UPolygonSelectionMechanic> SelectionMechanic = nullptr;
+
+	TArray<FName> Selection;
+	FName GetCurrentBone() const;
+
+	TMap<int32, int32> PolyGroupToBone;
+	TMap<int32, int32> PolyGroupToSubMesh;
+	
+protected:
+
+	int32 GetHoveredPolyGroup(const FInputDeviceRay& InRay) const;
+	
+	UPROPERTY()
+	TObjectPtr<USingleClickInputBehavior> LeftClickBehavior = nullptr;
+
+	UPROPERTY()
+	TObjectPtr<UPreviewMesh> PreviewMesh = nullptr;
+
+	UPROPERTY()
+	TWeakObjectPtr<UWorld> TargetWorld = nullptr;
+
+	UPROPERTY()
+	TWeakObjectPtr<USkeletalMeshEditorContextObjectBase> EditorContext = nullptr;
+	UPROPERTY()
+	TWeakObjectPtr<UPersonaEditorModeManagerContext> PersonaModeManagerContext = nullptr;
+
+	TArray<TWeakObjectPtr<USkeletalMeshComponent>> SubComponents;
+	
+private:
+	TArray<int32> GetSelectedBoneIndexes() const;
+
+	TSharedPtr<UE::Geometry::FDynamicMeshAABBTree3> MeshSpatial;
+	UE::Geometry::FDynamicMeshAABBTree3& GetSpatial();
+	bool bSpatialDirty = true;
+
+	int32 CurrentPolyGroup = INDEX_NONE;
+	int32 CurrentBone = INDEX_NONE;
+	
+	mutable FVector2D LastScreenPosition = FVector2D::ZeroVector;
+	mutable int32 LastHit = INDEX_NONE;
+
+	// setup
+	void SetupProperties();
+	void SetupPreviewMesh();
+	void SetupBehaviors();
+	void SetupPolyGroupSelection();
+
+	void GenerateDMCSubSkeletalMeshes();
+	void CleanDMCSubSkeletalMeshes();
+
+	void UpdateVisualization();
+	
+	TFunction<void()> PendingFunction;
+	TFunction<void()> ReleasePendingFunction;
+
+	TUniquePtr<UE::Geometry::FGroupTopology> Topology = nullptr;
+
+	TWeakObjectPtr<USkeletalMeshComponent> SourceSkeletalMeshComponent = nullptr;
+	TWeakObjectPtr<USkeletalMesh> SourceSkeletalMesh = nullptr;
+
+	TUniquePtr<UE::Geometry::FPolygroupSet> ActiveGroupSet;
+	TUniquePtr<UE::Geometry::FDynamicMeshPolygroupAttribute> PolygroupAttrib;
+
+	friend class UDirectMeshControlProperties;
+};
+
+/**
+ * UDirectMeshControlProperties
+ */
+
+UENUM()
+enum class EDMCAction
+{
+	None,
+	Select,
+	Assign,
+	Pose
+};
+
+UCLASS()
+class DIRECTMESHCONTROL_API UDirectMeshControlProperties : public UInteractiveToolPropertySet
+{
+	GENERATED_BODY()
+public:
+	void Initialize(UDirectMeshControlTool* InParentTool);
+
+	UPROPERTY(EditAnywhere, Category = Display, meta = (ClampMin = "0.0", ClampMax = "1.0", UIMin = "0.0", UIMax = "1.0"))
+	float Alpha = 0.5f;
+
+	UPROPERTY(EditAnywhere, Category = Display)
+	FVector Offset;
+
+	UPROPERTY(EditAnywhere, Category = Display)
+	FVector OffsetSub;
+
+	UPROPERTY(EditAnywhere, Category = Display)
+	bool bSubMeshes = true;
+	
+	UPROPERTY(EditAnywhere, Category = Display)
+	bool bAttachToActor = true;
+	
+	UPROPERTY(EditAnywhere, Category = Display)
+	bool bHitProxy = true;
+	
+	UPROPERTY(EditAnywhere, Category = Display)
+	bool bHoverThruMaterial = true;
+	
+	UPROPERTY(EditAnywhere, Category = Action)
+	EDMCAction Action = EDMCAction::Pose;
+
+	UFUNCTION(CallInEditor, Category = Action)
+	void Assign();
+	
+	void DefaultAssign();
+	
+private:
+	TWeakObjectPtr<UDirectMeshControlTool> ParentTool = nullptr;
+};
diff --git a/Engine/Plugins/Experimental/Animation/DirectMeshControl/Source/DirectMeshControl/Public/Tools/DirectMeshPolygroupTool.h b/Engine/Plugins/Experimental/Animation/DirectMeshControl/Source/DirectMeshControl/Public/Tools/DirectMeshPolygroupTool.h
new file mode 100644
--- /dev/null
+++ b/Engine/Plugins/Experimental/Animation/DirectMeshControl/Source/DirectMeshControl/Public/Tools/DirectMeshPolygroupTool.h
@@ -0,0 +1,160 @@
+// Copyright Epic Games, Inc. All Rights Reserved.
+
+#pragma once
+
+#include "CoreMinimal.h"
+#include "BaseTools/SingleTargetWithSelectionTool.h"
+#include "PreviewMesh.h"
+#include "ModelingOperators.h"
+#include "MeshOpPreviewHelpers.h"
+#include "PropertySets/PolygroupLayersProperties.h"
+#include "DynamicMesh/MeshSharingUtil.h"
+#include "DirectMeshPolygroupTool.generated.h"
+
+// predeclaration
+class UDirectMeshPolygroupTool;
+class FDirectMeshPolygroupOp;
+class UPreviewGeometry;
+namespace UE::Geometry
+{
+	class FDynamicMesh3;
+	struct FDynamicSubmesh3;
+}
+
+/**
+ *
+ */
+UCLASS()
+class DIRECTMESHCONTROL_API UDirectMeshPolygroupToolBuilder : public USingleTargetWithSelectionToolBuilder
+{
+	GENERATED_BODY()
+public:
+	virtual USingleTargetWithSelectionTool* CreateNewTool(const FToolBuilderState & SceneState) const override;
+	virtual bool RequiresInputSelection() const override { return false; }
+	virtual bool CanBuildTool(const FToolBuilderState& SceneState) const override;
+};
+
+UCLASS()
+class DIRECTMESHCONTROL_API UDirectMeshPolygroupToolProperties : public UInteractiveToolPropertySet
+{
+	GENERATED_BODY()
+
+public:
+
+	/** Bias for Quads that are adjacent to already-discovered Quads. Set to 0 to disable.  */
+	UPROPERTY(EditAnywhere, Category = Quads, meta = (UIMin = 0, UIMax = 5, EditCondition = "bTryUsingQuads == true"))
+	float QuadAdjacencyWeight = 1.0;
+
+	/** Set to values below 1 to ignore less-likely triangle pairings */
+	UPROPERTY(EditAnywhere, Category = Quads, meta = (AdvancedDisplay, UIMin = 0, UIMax = 1, EditCondition = "bTryUsingQuads == true"))
+	float QuadMetricClamp = 1.0;
+
+	/** Iteratively repeat quad-searching in uncertain areas, to try to slightly improve results */
+	UPROPERTY(EditAnywhere, Category = Quads, meta = (AdvancedDisplay, UIMin = 1, UIMax = 5, EditCondition = "bTryUsingQuads == true"))
+	int QuadSearchRounds = 1;
+
+	/** If true, polygroup borders will not cross existing UV seams */
+	UPROPERTY(EditAnywhere, Category = Quads, meta = (EditCondition = "bTryUsingQuads == true"))
+	bool bRespectUVSeams = false;
+
+	/** If true, polygroup borders will not cross existing hard normal seams */
+	UPROPERTY(EditAnywhere, Category = Quads, meta = (EditCondition = "bTryUsingQuads == true"))
+	bool bRespectHardNormals = false;
+
+	/** Minimum number of triangles to include in a group. Any group containing fewer triangles will be merged with an adjacent group (if possible). */
+	UPROPERTY(EditAnywhere, Category = Filtering, meta = (UIMin = "1", UIMax = "100", ClampMin = "1", ClampMax = "10000"))
+	int32 MinGroupSize = 2;
+	
+	/** Display each group with a different auto-generated color */
+	UPROPERTY(EditAnywhere, Category = Display)
+	bool bShowGroupColors = true;
+
+	/** If true, quads are used to defined skinned areas first */
+	UPROPERTY(EditAnywhere, Category = Options)
+	bool bTryUsingQuads = true;
+
+	/** Bones to remove from the polygroup computation */
+	UPROPERTY(EditAnywhere, Category = Options)
+	TArray<FName> BonesToRemove;
+	
+	/** Name of the output layer */
+	UPROPERTY(EditAnywhere, Category = Output)
+	FName LayerName = "dmc-polygroup";
+};
+
+
+
+
+UCLASS()
+class DIRECTMESHCONTROL_API UDirectMeshPolygroupOperatorFactory : public UObject, public UE::Geometry::IDynamicMeshOperatorFactory
+{
+	GENERATED_BODY()
+
+public:
+	// IDynamicMeshOperatorFactory API
+	virtual TUniquePtr<UE::Geometry::FDynamicMeshOperator> MakeNewOperator() override;
+
+	UPROPERTY()
+	TObjectPtr<UDirectMeshPolygroupTool> DirectMeshPolygroupTool;  // back pointer
+};
+
+/**
+ *
+ */
+UCLASS()
+class DIRECTMESHCONTROL_API UDirectMeshPolygroupTool : public USingleTargetWithSelectionTool
+{
+	GENERATED_BODY()
+
+public:
+	UDirectMeshPolygroupTool();
+
+	virtual void Setup() override;
+	virtual void OnShutdown(EToolShutdownType ShutdownType) override;
+
+	virtual void OnTick(float DeltaTime) override;
+
+	virtual bool HasCancel() const override { return true; }
+	virtual bool HasAccept() const override { return true; }
+	virtual bool CanAccept() const override;
+
+	// update parameters in DirectMeshPolygroupOp based on current Settings
+	void UpdateOpParameters(FDirectMeshPolygroupOp& DirectMeshPolygroupOp) const;
+
+protected:
+	
+	void OnSettingsModified();
+
+protected:
+	UPROPERTY()
+	TObjectPtr<UDirectMeshPolygroupToolProperties> Properties;
+
+	UPROPERTY()
+	TObjectPtr<UMeshOpPreviewWithBackgroundCompute> PreviewCompute = nullptr;
+
+	UPROPERTY()
+	TObjectPtr<UPreviewGeometry> PreviewGeometry = nullptr;
+
+	// If a selection was provided (bUsingSelection = true), UnmodifiedAreaPreviewMesh is used to render the unmodified (non-selected) part of the input mesh,
+	// as the PreviewCompute input mesh will be limited to the input selected area
+	UPROPERTY()
+	TObjectPtr<UPreviewMesh> UnmodifiedAreaPreviewMesh = nullptr;
+
+protected:
+	// copy of the input mesh (possibly does not need to be SharedPtr anymore now that FSharedConstDynamicMesh3 is used below, but then needs PimplPtr...)
+	TSharedPtr<UE::Geometry::FDynamicMesh3, ESPMode::ThreadSafe> OriginalDynamicMesh;
+
+	// If there is an active selection, bUsingSelection will be true and the ROI and Submesh will be initialized
+	bool bUsingSelection = false;
+	TSharedPtr<TSet<int32>> SelectionTriangleROI;
+	TSharedPtr<UE::Geometry::FDynamicSubmesh3, ESPMode::ThreadSafe> OriginalSubmesh;
+
+	// The mesh passed to the Compute operator to base result on, will either be OriginalDynamicMesh or the submesh from OriginalSubmesh if using a selection
+	TSharedPtr<UE::Geometry::FSharedConstDynamicMesh3> ComputeOperatorSharedMesh;
+
+	// current set of detected polygroup edges, relative to ComputeOperatorSharedMesh
+	TArray<int> PolygonEdges;
+	
+	void UpdateVisualization();
+
+};
diff --git a/Engine/Plugins/Experimental/Animation/DirectMeshControl/Source/DirectMeshControlRig/DirectMeshControlRig.Build.cs b/Engine/Plugins/Experimental/Animation/DirectMeshControl/Source/DirectMeshControlRig/DirectMeshControlRig.Build.cs
new file mode 100644
--- /dev/null
+++ b/Engine/Plugins/Experimental/Animation/DirectMeshControl/Source/DirectMeshControlRig/DirectMeshControlRig.Build.cs
@@ -0,0 +1,29 @@
+﻿// Copyright Epic Games, Inc. All Rights Reserved.
+
+using UnrealBuildTool;
+
+public class DirectMeshControlRig : ModuleRules
+{
+	public DirectMeshControlRig(ReadOnlyTargetRules Target) : base(Target)
+	{
+		PCHUsage = ModuleRules.PCHUsageMode.UseExplicitOrSharedPCHs;
+		
+		PrivateDependencyModuleNames.AddRange(
+			new string[]
+			{
+				"DirectMeshControl",
+				"AnimationCore",
+				"ControlRig",
+				"ComputeFramework",
+				"Core",
+				"CoreUObject",
+				"DynamicMesh",
+				"Engine",
+				"GeometryCore",
+				"MeshConversion",
+				"OptimusCore",
+				"RigVM"
+			}
+		);
+	}
+}
diff --git a/Engine/Plugins/Experimental/Animation/DirectMeshControl/Source/DirectMeshControlRig/Private/DirectMeshControlProxyWrapper.cpp b/Engine/Plugins/Experimental/Animation/DirectMeshControl/Source/DirectMeshControlRig/Private/DirectMeshControlProxyWrapper.cpp
new file mode 100644
--- /dev/null
+++ b/Engine/Plugins/Experimental/Animation/DirectMeshControl/Source/DirectMeshControlRig/Private/DirectMeshControlProxyWrapper.cpp
@@ -0,0 +1,3 @@
+﻿// Copyright Epic Games, Inc. All Rights Reserved.
+
+#include "DirectMeshControlProxyWrapper.h"
diff --git a/Engine/Plugins/Experimental/Animation/DirectMeshControl/Source/DirectMeshControlRig/Private/DirectMeshControlRigModule.cpp b/Engine/Plugins/Experimental/Animation/DirectMeshControl/Source/DirectMeshControlRig/Private/DirectMeshControlRigModule.cpp
new file mode 100644
--- /dev/null
+++ b/Engine/Plugins/Experimental/Animation/DirectMeshControl/Source/DirectMeshControlRig/Private/DirectMeshControlRigModule.cpp
@@ -0,0 +1,60 @@
+﻿// Copyright Epic Games, Inc. All Rights Reserved.
+
+#include "DirectMeshControlRigModule.h"
+
+#include "ControlRigGizmoLibrary.h"
+#include "DirectMeshControlComponent.h"
+#include "DirectMeshControlProxyWrapper.h"
+#include "Engine/SkeletalMesh.h"
+
+#define LOCTEXT_NAMESPACE "FDirectMeshControlRigModule"
+
+void FDirectMeshControlRigModule::StartupModule()
+{
+	FCoreDelegates::GetOnPostEngineInit().AddRaw(this, &FDirectMeshControlRigModule::OnPostEngineInit);
+}
+
+void FDirectMeshControlRigModule::ShutdownModule()
+{
+	FCoreDelegates::GetOnPostEngineInit().RemoveAll(this);
+}
+
+void FDirectMeshControlRigModule::OnPostEngineInit() const
+{
+	UE_LOG(LogTemp, Warning, TEXT("FDirectMeshControlRigModule::OnPostEngineInit"));
+	
+	FShapeProxyComponentProviderRegistry& Registry = FShapeProxyComponentProviderRegistry::Get();
+	Registry.RegisterProvider<UDirectMeshControlProxy, UDirectMeshControlComponent>(
+		[](UObject* Proxy, UObject* Outer) -> UPrimitiveComponent*
+		{
+			UDirectMeshControlProxy* DMCProxy = Cast<UDirectMeshControlProxy>(Proxy);
+			if (USkeletalMesh* SkeletalMesh = DMCProxy ? DMCProxy->GetTypedOuter<USkeletalMesh>() : nullptr)
+			{
+				UE_LOG(LogTemp, Warning, TEXT("new proxy for %s"), *SkeletalMesh->GetName());
+				
+				const FName UniqueName = MakeUniqueObjectName(Outer, UDirectMeshControlComponent::StaticClass(), TEXT("DirectMeshControlComponent_CtrlProxy"));
+				UDirectMeshControlComponent* DMCComponent = NewObject<UDirectMeshControlComponent>(Outer, UniqueName);
+				DMCComponent->SetSkeletalMesh(SkeletalMesh);
+				return DMCComponent;
+			}
+			return nullptr; 
+		},
+		[](UObject* Proxy, UObject* Component)
+		{
+			UDirectMeshControlProxy* DMCProxy = Cast<UDirectMeshControlProxy>(Proxy);
+			USkeletalMesh* SkeletalMesh = DMCProxy ? DMCProxy->GetTypedOuter<USkeletalMesh>() : nullptr;
+			UDirectMeshControlComponent* DMCComponent = Cast<UDirectMeshControlComponent>(Component);
+			if (SkeletalMesh && DMCComponent && DMCComponent->GetSkeletalMeshAsset() != SkeletalMesh)
+			{
+				UE_LOG(LogTemp, Warning, TEXT("update component with %s"), *SkeletalMesh->GetName());
+				
+				DMCComponent->SetSkeletalMesh(SkeletalMesh);
+				return true;
+			}
+			return false;
+		});
+}
+
+#undef LOCTEXT_NAMESPACE
+	
+IMPLEMENT_MODULE(FDirectMeshControlRigModule, DirectMeshControlRig)
\ No newline at end of file
diff --git a/Engine/Plugins/Experimental/Animation/DirectMeshControl/Source/DirectMeshControlRig/Private/Units/RigUnit_DirectMeshControl.cpp b/Engine/Plugins/Experimental/Animation/DirectMeshControl/Source/DirectMeshControlRig/Private/Units/RigUnit_DirectMeshControl.cpp
new file mode 100644
--- /dev/null
+++ b/Engine/Plugins/Experimental/Animation/DirectMeshControl/Source/DirectMeshControlRig/Private/Units/RigUnit_DirectMeshControl.cpp
@@ -0,0 +1,221 @@
+// Copyright Epic Games, Inc. All Rights Reserved.
+
+#include "Units/RigUnit_DirectMeshControl.h"
+
+#include "Units/Execution/RigUnit_DynamicHierarchy.h"
+#include "ControlRig.h"
+#include "ControlRigGizmoActor.h"
+#include "DirectMeshControlProxyWrapper.h"
+#include "DirectMeshControlSettings.h"
+#include "DirectMeshControlUtilities.h"
+#include "IControlRigObjectBinding.h"
+#include "MeshDescriptionToDynamicMesh.h"
+#include "OptimusDeformer.h"
+#include "OptimusDeformerInstance.h"
+#include "Engine/SkeletalMesh.h"
+#include "Materials/MaterialInstanceDynamic.h"
+#include "Polygroups/GroupSetAdapter.h"
+#include "Polygroups/PolygroupUtil.h"
+#include "Util/ColorConstants.h"
+
+#include UE_INLINE_GENERATED_CPP_BY_NAME(RigUnit_DirectMeshControl)
+
+namespace UE::DirectMeshControl::Private
+{
+
+bool GetDynamicMesh(const USkeletalMesh* InSkeletalMesh, FDynamicMesh3& OutDynaMesh)
+{
+	if (!InSkeletalMesh)
+	{
+		return false;
+	}
+	
+	constexpr int32 LOD = 0;
+	FMeshDescription* MeshDescription = InSkeletalMesh->HasMeshDescription(LOD) ? InSkeletalMesh->GetMeshDescription(LOD) : nullptr;
+	if (!MeshDescription)
+	{
+		return false;
+	}
+
+	FMeshDescriptionToDynamicMesh Converter;
+	Converter.bTransformVertexColorsLinearToSRGB = false;
+	Converter.bVIDsFromNonManifoldMeshDescriptionAttr = true;
+	constexpr bool bWantTangents = true;
+	Converter.Convert(MeshDescription, OutDynaMesh, bWantTangents);
+
+	return true;
+}
+	
+}
+
+FRigUnit_SetupShapeLibraryFromLayer_Execute()
+{
+	using namespace UE::DirectMeshControl::Private;
+	using namespace UE::Geometry;
+	
+	FString ErrorMessage;
+
+	GroupNames.Reset();
+	
+	constexpr bool bAllowOnlyConstructionEvent = true;
+	if (!FRigUnit_DynamicHierarchyBase::IsValidToRunInContext(ExecuteContext, bAllowOnlyConstructionEvent, &ErrorMessage))
+	{
+		if(!ErrorMessage.IsEmpty())
+		{
+			UE_CONTROLRIG_RIGUNIT_REPORT_ERROR(TEXT("%s"), *ErrorMessage);
+		}
+		return;
+	}
+
+	if (!ExecuteContext.OnAddShapeLibraryDelegate.IsBound())
+	{
+		UE_CONTROLRIG_RIGUNIT_REPORT_ERROR(TEXT("OnAddShapeLibraryDelegate not bound"));
+		return;
+	}
+	
+	if (LayerName == NAME_None)
+	{
+		return;
+	}
+
+	URigHierarchyController* Controller = ExecuteContext.Hierarchy->GetController(true);
+	if (!Controller)
+	{
+		return;
+	}
+	
+	USkeletalMeshComponent* SkeletalMeshComponent = ExecuteContext.UnitContext.RequestDataSource<USkeletalMeshComponent>(UControlRig::OwnerComponent);
+	USkeletalMesh* SkeletalMesh = SkeletalMeshComponent ? SkeletalMeshComponent->GetSkeletalMeshAsset() : nullptr;
+
+	FDynamicMesh3 DynaMesh;
+	if (!GetDynamicMesh(SkeletalMesh, DynaMesh))
+	{
+		return;
+	}
+
+	FDynamicMeshTriangleLabelAttribute* DMCLabelAttr = FindTriangleLabelLayerByName(DynaMesh, LayerName);
+	if (!DMCLabelAttr)
+	{
+		return;
+	}
+	
+	FTriangleAttributeAdapter LabelAdapter(&DynaMesh, DMCLabelAttr);
+	{
+		TSet<FName> UniqueNames;
+		for (const int32 TriID : DynaMesh.TriangleIndicesItr())
+		{
+			UniqueNames.Add(LabelAdapter.GetValue(TriID));
+		}
+		GroupNames = UniqueNames.Array();
+	}
+
+	if (!GroupNames.IsEmpty())
+	{
+		const FSoftObjectPath SkeletalMeshPath(SkeletalMesh);
+		FString LibraryName(SkeletalMeshPath.GetSubPathUtf8String());
+		LibraryName += LayerName.ToString();
+	
+		UControlRig* ControlRig = ExecuteContext.ControlRig;
+		const TMap<FString, FString>& LibraryNames = ControlRig->GetShapeLibraryNameMap();
+		if (!LibraryNames.Contains(LibraryName))
+		{
+			UObjectBase* Found = FindObjectWithOuter(SkeletalMesh, UControlRigShapeLibrary::StaticClass(), LayerName);
+			if (!Found)
+			{
+				// replace with a dmc factory placeholder
+				UControlRigShapeLibrary* NewDMCShapeLibrary = NewObject<UControlRigShapeLibrary>(SkeletalMesh, LayerName, RF_NoFlags);
+				NewDMCShapeLibrary->DefaultShape.ShapeName = NAME_None;
+				NewDMCShapeLibrary->Shapes.Reserve(GroupNames.Num());
+				NewDMCShapeLibrary->DefaultMaterial = UE::DMC::GetMaterial();
+				NewDMCShapeLibrary->XRayMaterial = UE::DMC::GetMaterial();
+
+				const FGroupSubMeshes& GroupSubMeshes = UE::DMC::GetSubMeshes(SkeletalMesh, &DynaMesh, LayerName);
+				const TArray<TObjectPtr<USkeletalMesh>>& SubMeshes = GroupSubMeshes.GetSubSkeletalMeshes();
+		
+				TArray<TPair<int32, int32>> GroupIdToSubMesh = GroupSubMeshes.GetGroupIdToSubMesh().Array();
+				Algo::SortBy(GroupIdToSubMesh, [&LabelAdapter](const TPair<int32, int32>& In)
+				{
+					return LabelAdapter.GetValueFromGroup(In.Key);
+				},
+				FNameLexicalLess());
+				
+				const UDirectMeshControlSettings* DirectMeshControlSettings = GetMutableDefault<UDirectMeshControlSettings>();
+				constexpr FLinearColor DefaultOverlay(0.25f, 0.25f, 0.25f, 0.25f);
+				const FLinearColor& OverlayColor = DirectMeshControlSettings ? DirectMeshControlSettings->OverlayColor : DefaultOverlay;
+				for (int32 Index = 0; Index < SubMeshes.Num(); ++Index)
+				{
+					const auto& [GroudID, SubMeshIndex] = GroupIdToSubMesh[Index];
+					
+					ensure(SubMeshes.IsValidIndex(SubMeshIndex));
+					
+					USkeletalMesh* SubSkeletalMesh = SubMeshes.IsValidIndex(SubMeshIndex) ? SubMeshes[SubMeshIndex].Get() : nullptr;
+					if (ensure(SubSkeletalMesh))
+					{
+						FControlRigShapeDefinition& ShapeDefinition = NewDMCShapeLibrary->Shapes.Emplace_GetRef();
+						ShapeDefinition.ShapeName = LabelAdapter.GetValueFromGroup(GroudID);
+						// ShapeDefinition.ShapeProxy = SubSkeletalMesh;
+						ShapeDefinition.ShapeProxy = NewObject<UDirectMeshControlProxy>(SubSkeletalMesh);
+						ShapeDefinition.PostSetupFunction = [GroudID, OverlayColor](AControlRigShapeActor* InShape)
+						{
+							if (USkeletalMeshComponent* DMCComponent = InShape ? Cast<USkeletalMeshComponent>(InShape->ProxyComponent.Get()) : nullptr)
+							{
+								{
+									// FIXME culling
+									DMCComponent->bAllowCullDistanceVolume = false;
+									DMCComponent->bNeverDistanceCull = true;
+									DMCComponent->BoundsScale = 1000.0f;
+								}
+							
+								// prepare source skm
+								UControlRig* ControlRig = InShape->ControlRig.Get();
+								const TSharedPtr<IControlRigObjectBinding>& ObjectBinding = ControlRig ? ControlRig->GetObjectBinding() : nullptr;
+								USkeletalMeshComponent* BoundSkeletalMeshComponent =
+										ObjectBinding ? Cast<USkeletalMeshComponent>(ObjectBinding->GetBoundObject()) : nullptr;
+								if (BoundSkeletalMeshComponent)
+								{
+									BoundSkeletalMeshComponent->SetAlwaysUseMeshDeformer(true);
+									BoundSkeletalMeshComponent->SetForcedLOD(1);
+								}
+								
+								// add deformer
+								DMCComponent->SetMeshDeformer(UE::DMC::GetDeformer());
+								
+								// setup material
+								if (UE::DMC::UseMaterialHover())
+								{
+									if (UMaterialInstanceDynamic* MaterialInstance = Cast<UMaterialInstanceDynamic>(DMCComponent->GetMaterial(0)))
+									{
+										InShape->HoveredParameterName = TEXT("Hovered");
+										MaterialInstance->SetScalarParameterValue(TEXT("Hovered"), 0.f);
+										
+										InShape->HoveredColorParameterName = TEXT("HoveredColor");
+										InShape->SetHoveredColor(OverlayColor);
+									}
+								}
+								else
+								{
+									UOptimusDeformerInstance* MeshDeformerInstance = Cast<UOptimusDeformerInstance>(DMCComponent->GetMeshDeformerInstance());
+									if (ensure(MeshDeformerInstance))
+									{
+										FColor Color = LinearColors::SelectFColor(GroudID);
+										Color.A = 64.f;
+										ensure(MeshDeformerInstance->SetLinearColorVariable(UE::DMC::GetColorVarName(), FLinearColor(Color)));
+									}
+								}
+							}
+						};
+					}
+				}
+
+				constexpr bool bDoNotLogResults = false;
+				ExecuteContext.OnAddShapeLibraryDelegate.Execute(&ExecuteContext, LibraryName, NewDMCShapeLibrary, bDoNotLogResults);
+			}
+			else
+			{
+				UControlRigShapeLibrary* FoundDMCShapeLibrary = static_cast<UControlRigShapeLibrary*>(Found);
+				constexpr bool bDoNotLogResults = false;
+				ExecuteContext.OnAddShapeLibraryDelegate.Execute(&ExecuteContext, LibraryName, FoundDMCShapeLibrary, bDoNotLogResults);
+			}
+		}
+	}
+}
diff --git a/Engine/Plugins/Experimental/Animation/DirectMeshControl/Source/DirectMeshControlRig/Public/DirectMeshControlProxyWrapper.h b/Engine/Plugins/Experimental/Animation/DirectMeshControl/Source/DirectMeshControlRig/Public/DirectMeshControlProxyWrapper.h
new file mode 100644
--- /dev/null
+++ b/Engine/Plugins/Experimental/Animation/DirectMeshControl/Source/DirectMeshControlRig/Public/DirectMeshControlProxyWrapper.h
@@ -0,0 +1,16 @@
+﻿// Copyright Epic Games, Inc. All Rights Reserved.
+
+#pragma once
+
+#include "UObject/Object.h"
+
+#include "DirectMeshControlProxyWrapper.generated.h"
+
+/**
+ * 
+ */
+UCLASS(transient)
+class DIRECTMESHCONTROLRIG_API UDirectMeshControlProxy : public UObject
+{
+	GENERATED_BODY()
+};
diff --git a/Engine/Plugins/Experimental/Animation/DirectMeshControl/Source/DirectMeshControlRig/Public/DirectMeshControlRigModule.h b/Engine/Plugins/Experimental/Animation/DirectMeshControl/Source/DirectMeshControlRig/Public/DirectMeshControlRigModule.h
new file mode 100644
--- /dev/null
+++ b/Engine/Plugins/Experimental/Animation/DirectMeshControl/Source/DirectMeshControlRig/Public/DirectMeshControlRigModule.h
@@ -0,0 +1,17 @@
+﻿// Copyright Epic Games, Inc. All Rights Reserved.
+
+#pragma once
+
+#include "Modules/ModuleManager.h"
+
+class FDirectMeshControlRigModule : public IModuleInterface
+{
+public:
+
+	/** IModuleInterface implementation */
+	virtual void StartupModule() override;
+	virtual void ShutdownModule() override;
+	
+private:
+	void OnPostEngineInit() const;
+};
diff --git a/Engine/Plugins/Experimental/Animation/DirectMeshControl/Source/DirectMeshControlRig/Public/Units/RigUnit_DirectMeshControl.h b/Engine/Plugins/Experimental/Animation/DirectMeshControl/Source/DirectMeshControlRig/Public/Units/RigUnit_DirectMeshControl.h
new file mode 100644
--- /dev/null
+++ b/Engine/Plugins/Experimental/Animation/DirectMeshControl/Source/DirectMeshControlRig/Public/Units/RigUnit_DirectMeshControl.h
@@ -0,0 +1,37 @@
+// Copyright Epic Games, Inc. All Rights Reserved.
+
+#pragma once
+
+#include "Units/Execution/RigUnit_DynamicHierarchy.h"
+
+#include "RigUnit_DirectMeshControl.generated.h"
+
+
+/**
+* DMC FRigUnit_SetupShapeLibraryFromLayer
+*/
+USTRUCT(meta=(DisplayName="Set Shape Library from Layer", Keywords="Direct Mesh Control, Surface Selection", Varying))
+struct FRigUnit_SetupShapeLibraryFromLayer : public FRigUnit_DynamicHierarchyBaseMutable
+{
+	GENERATED_BODY()
+
+	FRigUnit_SetupShapeLibraryFromLayer()
+	{
+		LayerName = "dmc-polygroup";
+	}
+
+	RIGVM_METHOD()
+	virtual void Execute() override;
+
+	/*
+	 * The layer to extract
+	 */
+	UPROPERTY(meta = (Input))
+	FName LayerName;
+	
+	/*
+	 * The extracted groups
+	 */
+	UPROPERTY(meta = (Output))
+	TArray<FName> GroupNames;
+};
diff --git a/Engine/Plugins/Runtime/GeometryProcessing/Source/DynamicMesh/Public/Polygroups/GroupSetAdapter.h b/Engine/Plugins/Runtime/GeometryProcessing/Source/DynamicMesh/Public/Polygroups/GroupSetAdapter.h
--- a/Engine/Plugins/Runtime/GeometryProcessing/Source/DynamicMesh/Public/Polygroups/GroupSetAdapter.h
+++ b/Engine/Plugins/Runtime/GeometryProcessing/Source/DynamicMesh/Public/Polygroups/GroupSetAdapter.h
@@ -36,7 +36,7 @@ public:
 	
 protected:
 	/** Initialize a GroupSetAdapter for the given Mesh. */
-	ITriangleGroupSetAdapter(const FDynamicMesh3* InMesh);
+	UE_API ITriangleGroupSetAdapter(const FDynamicMesh3* InMesh);
 	/** The mesh this GroupSetAdapter references. */
 	const FDynamicMesh3* Mesh = nullptr;
 	/** Max Group ID. */
