Direct Mesh Control: fix CIS issue

#rnx
#Synced-CL 52542108

diff --git a/Engine/Plugins/Experimental/Animation/DirectMeshControl/Source/DirectMeshControl/Private/Tools/DirectMeshPolygroupTool.cpp b/Engine/Plugins/Experimental/Animation/DirectMeshControl/Source/DirectMeshControl/Private/Tools/DirectMeshPolygroupTool.cpp
--- a/Engine/Plugins/Experimental/Animation/DirectMeshControl/Source/DirectMeshControl/Private/Tools/DirectMeshPolygroupTool.cpp
+++ b/Engine/Plugins/Experimental/Animation/DirectMeshControl/Source/DirectMeshControl/Private/Tools/DirectMeshPolygroupTool.cpp
@@ -81,7 +81,7 @@ public:
 	/**
 	 * Find Polygroups based on skin weights
 	 */
-	bool FindPolygroupsFromSkinWeights(const FPolygroupsFromSkinWeightsParams& InParams = FPolygroupsFromSkinWeightsParams())
+	bool FindPolygroupsFromSkinWeights(const FPolygroupsFromSkinWeightsParams& InParams)
 	{
 		using namespace UE::AnimationCore;
 		
