Direct Mesh Control: force DMC components animation updates in editor

#rnx

#Synced-CL 52675748

diff --git a/Engine/Plugins/Experimental/Animation/DirectMeshControl/Source/DirectMeshControlRig/Private/Units/RigUnit_DirectMeshControl.cpp b/Engine/Plugins/Experimental/Animation/DirectMeshControl/Source/DirectMeshControlRig/Private/Units/RigUnit_DirectMeshControl.cpp
--- a/Engine/Plugins/Experimental/Animation/DirectMeshControl/Source/DirectMeshControlRig/Private/Units/RigUnit_DirectMeshControl.cpp
+++ b/Engine/Plugins/Experimental/Animation/DirectMeshControl/Source/DirectMeshControlRig/Private/Units/RigUnit_DirectMeshControl.cpp
@@ -202,6 +202,8 @@ FRigUnit_SetupShapeLibraryFromLayer_Execute()
 										ensure(MeshDeformerInstance->SetLinearColorVariable(UE::DMC::GetColorVarName(), FLinearColor(Color)));
 									}
 								}
+								
+								DMCComponent->SetUpdateAnimationInEditor(true);
 							}
 						};
 					}
