Direct Mesh Control: remove useless logs

#rnx
#Synced-CL 52675748

diff --git a/Engine/Plugins/Experimental/Animation/DirectMeshControl/Source/DirectMeshControlRig/Private/DirectMeshControlRigModule.cpp b/Engine/Plugins/Experimental/Animation/DirectMeshControl/Source/DirectMeshControlRig/Private/DirectMeshControlRigModule.cpp
--- a/Engine/Plugins/Experimental/Animation/DirectMeshControl/Source/DirectMeshControlRig/Private/DirectMeshControlRigModule.cpp
+++ b/Engine/Plugins/Experimental/Animation/DirectMeshControl/Source/DirectMeshControlRig/Private/DirectMeshControlRigModule.cpp
@@ -21,17 +21,13 @@ void FDirectMeshControlRigModule::ShutdownModule()
 
 void FDirectMeshControlRigModule::OnPostEngineInit() const
 {
-	UE_LOGF(LogTemp, Warning, "FDirectMeshControlRigModule::OnPostEngineInit");
-	
 	FShapeProxyComponentProviderRegistry& Registry = FShapeProxyComponentProviderRegistry::Get();
 	Registry.RegisterProvider<UDirectMeshControlProxy, UDirectMeshControlComponent>(
 		[](UObject* Proxy, UObject* Outer) -> UPrimitiveComponent*
 		{
 			UDirectMeshControlProxy* DMCProxy = Cast<UDirectMeshControlProxy>(Proxy);
 			if (USkeletalMesh* SkeletalMesh = DMCProxy ? DMCProxy->GetTypedOuter<USkeletalMesh>() : nullptr)
-			{
-				UE_LOGF(LogTemp, Warning, "new proxy for %ls", *SkeletalMesh->GetName());
-				
+			{	
 				const FName UniqueName = MakeUniqueObjectName(Outer, UDirectMeshControlComponent::StaticClass(), TEXT("DirectMeshControlComponent_CtrlProxy"));
 				UDirectMeshControlComponent* DMCComponent = NewObject<UDirectMeshControlComponent>(Outer, UniqueName);
 				DMCComponent->SetSkeletalMesh(SkeletalMesh);
@@ -46,8 +42,6 @@ void FDirectMeshControlRigModule::OnPostEngineInit() const
 			UDirectMeshControlComponent* DMCComponent = Cast<UDirectMeshControlComponent>(Component);
 			if (SkeletalMesh && DMCComponent && DMCComponent->GetSkeletalMeshAsset() != SkeletalMesh)
 			{
-				UE_LOGF(LogTemp, Warning, "update component with %ls", *SkeletalMesh->GetName());
-				
 				DMCComponent->SetSkeletalMesh(SkeletalMesh);
 				return true;
 			}
