Direct Mesh Control: documentation pass

#rnx
#Synced-CL 52722994

diff --git a/Engine/Plugins/Experimental/Animation/DirectMeshControl/Source/DirectMeshControl/Private/DMCMeshGenerationSubsystem.h b/Engine/Plugins/Experimental/Animation/DirectMeshControl/Source/DirectMeshControl/Private/DMCMeshGenerationSubsystem.h
--- a/Engine/Plugins/Experimental/Animation/DirectMeshControl/Source/DirectMeshControl/Private/DMCMeshGenerationSubsystem.h
+++ b/Engine/Plugins/Experimental/Animation/DirectMeshControl/Source/DirectMeshControl/Private/DMCMeshGenerationSubsystem.h
@@ -19,7 +19,7 @@ namespace UE::Geometry
 }
 
 /**
- *
+ * UDMCMeshGenerationManager is a persistent cache manager for (skeletal mesh, layer name) to sub-meshes mappings.
  */
 
 UCLASS(MinimalAPI)
@@ -28,18 +28,26 @@ class UDMCMeshGenerationManager : public UObject
 	GENERATED_BODY()
 
 public:
+	/** Clears all cached sub-mesh entries, releasing references to generated assets. */
 	UE_API virtual void Shutdown();
 
+	/**
+	 * Returns the cached FGroupSubMeshes for InSkeletalMesh / LayerName, rebuilding the entry when the DDC hash indicates the source data has changed.
+	 * @param InSkeletalMesh  Source skeletal mesh.
+	 * @param Mesh            Current dynamic mesh representation used to detect changes and rebuild.
+	 * @param LayerName       Triangle label layer that encodes polygroup IDs.
+	 * @return Reference to the up-to-date cached FGroupSubMeshes.
+	 */
 	UE_API const FGroupSubMeshes& GetSubMeshes(USkeletalMesh* InSkeletalMesh, const UE::Geometry::FDynamicMesh3* Mesh, const FName LayerName);
 	
 protected:
-	
+	/** The set of cached entries keyed by (SkeletalMesh, LayerName). */
 	UPROPERTY()
 	TSet<FSubSkeletalMeshData> SubSkeletalMeshes;
 };
 	
 /**
- *
+ * UDMCMeshGenerationSubsystem is an editor subsystem that provides global access to UDMCMeshGenerationManager.
  */
 
 UCLASS(MinimalAPI)
@@ -48,17 +56,28 @@ class UDMCMeshGenerationSubsystem : public UEditorSubsystem
 	GENERATED_BODY()
 
 public:
+	/** Returns the active UDMCMeshGenerationSubsystem instance, or nullptr if the engine is shutting down. */
 	static UDMCMeshGenerationSubsystem* Get();
+
+	/**
+	 * Creates the UDMCMeshGenerationManager and binds the engine-exit and pre-exit shutdown callbacks.
+	 * @param Collection  Subsystem collection (passed to the parent implementation).
+	 */
 	UE_API virtual void Initialize(FSubsystemCollectionBase& Collection) override;
+
+	/** Removes bound shutdown callbacks, calls Shutdown() on the manager, and nulls the @c GenerationManager pointer. */
 	UE_API virtual void Deinitialize() override;
 
+	/** The owned mesh generation manager. Accessed via Get()->GenerationManager. */
 	UPROPERTY()
 	TObjectPtr<UDMCMeshGenerationManager> GenerationManager = nullptr;
-	
+
 protected:
+	/** Sets bIsShuttingDown to true so that subsequent Get() calls return nullptr. */
 	UE_API virtual void OnShutdown();
 
 private:
+	/** Static flag set to true during engine shutdown to guard Get() against use-after-destroy. */
 	static UE_API bool bIsShuttingDown;
 };
 	
diff --git a/Engine/Plugins/Experimental/Animation/DirectMeshControl/Source/DirectMeshControl/Private/DirectMeshControlComponent.cpp b/Engine/Plugins/Experimental/Animation/DirectMeshControl/Source/DirectMeshControl/Private/DirectMeshControlComponent.cpp
--- a/Engine/Plugins/Experimental/Animation/DirectMeshControl/Source/DirectMeshControl/Private/DirectMeshControlComponent.cpp
+++ b/Engine/Plugins/Experimental/Animation/DirectMeshControl/Source/DirectMeshControl/Private/DirectMeshControlComponent.cpp
@@ -9,11 +9,6 @@
 #include "Engine/SkeletalMesh.h"
 #include "Rendering/SkeletalMeshRenderData.h"
 
-void UDirectMeshControlComponent::CreateRenderState_Concurrent(FRegisterComponentContext* Context)
-{
-	Super::CreateRenderState_Concurrent(Context);
-}
-
 void UDirectMeshControlComponent::SendRenderDynamicData_Concurrent()
 {
 #if WITH_STATE_STREAM_ACTOR
diff --git a/Engine/Plugins/Experimental/Animation/DirectMeshControl/Source/DirectMeshControl/Public/DirectMeshControlCommands.h b/Engine/Plugins/Experimental/Animation/DirectMeshControl/Source/DirectMeshControl/Public/DirectMeshControlCommands.h
--- a/Engine/Plugins/Experimental/Animation/DirectMeshControl/Source/DirectMeshControl/Public/DirectMeshControlCommands.h
+++ b/Engine/Plugins/Experimental/Animation/DirectMeshControl/Source/DirectMeshControl/Public/DirectMeshControlCommands.h
@@ -6,18 +6,19 @@
 #include "Framework/Commands/Commands.h"
 
 /**
- * 
+ * FDirectMeshControlCommands holds the registered UI commands for the Direct Mesh Control plugin.
  */
 class DIRECTMESHCONTROL_API FDirectMeshControlCommands : public TCommands<FDirectMeshControlCommands>
 {
 public:
 	FDirectMeshControlCommands();
 
+	/** Command that activates the "DMC Polygroups" tool for generating bone-based polygroups. */
 	TSharedPtr<FUICommandInfo> BeginDirectMeshPolygroupTool;
+	
+	/** Command that activates the "DMC Control" tool for per-polygroup pose and selection. */
 	TSharedPtr<FUICommandInfo> BeginDirectMeshControlTool;
 
-	/**
-	 * Initialize commands
-	 */
+	/** Initialize commands. */
 	virtual void RegisterCommands() override;
 };
\ No newline at end of file
diff --git a/Engine/Plugins/Experimental/Animation/DirectMeshControl/Source/DirectMeshControl/Public/DirectMeshControlComponent.h b/Engine/Plugins/Experimental/Animation/DirectMeshControl/Source/DirectMeshControl/Public/DirectMeshControlComponent.h
--- a/Engine/Plugins/Experimental/Animation/DirectMeshControl/Source/DirectMeshControl/Public/DirectMeshControlComponent.h
+++ b/Engine/Plugins/Experimental/Animation/DirectMeshControl/Source/DirectMeshControl/Public/DirectMeshControlComponent.h
@@ -9,14 +9,16 @@
 #define UE_API DIRECTMESHCONTROL_API
 
 /**
- * 
+ * UDirectMeshControlComponent is a transient skeletal mesh component that bypasses normal bone transform processing and applies mesh deformers directly.
+ * Used by the DMC tools to render per-polygroup visualization sub-meshes.
+ * Bone transforms are skipped so that the deformer drives the geometry entirely, enabling the per-polygroup overlay without interfering with the source skeleton.
  */
 UCLASS(transient)
 class UE_API UDirectMeshControlComponent : public USkeletalMeshComponent
 {
 	GENERATED_BODY()
  	
-	virtual void CreateRenderState_Concurrent(FRegisterComponentContext* Context) override;
+	/** UPrimitiveComponent override. */
 	virtual void SendRenderDynamicData_Concurrent() override;
 };
 
diff --git a/Engine/Plugins/Experimental/Animation/DirectMeshControl/Source/DirectMeshControl/Public/DirectMeshControlModule.h b/Engine/Plugins/Experimental/Animation/DirectMeshControl/Source/DirectMeshControl/Public/DirectMeshControlModule.h
--- a/Engine/Plugins/Experimental/Animation/DirectMeshControl/Source/DirectMeshControl/Public/DirectMeshControlModule.h
+++ b/Engine/Plugins/Experimental/Animation/DirectMeshControl/Source/DirectMeshControl/Public/DirectMeshControlModule.h
@@ -13,7 +13,7 @@ class FExtender;
 class FRigVMEditorBase;
 
 /**
- * 
+ * FDirectMeshControlModule is the main editor module for the Direct Mesh Control plugin.
  */
 class FDirectMeshControlModule : public IModuleInterface, public ISkeletalMeshModelingModeToolExtension
 {
diff --git a/Engine/Plugins/Experimental/Animation/DirectMeshControl/Source/DirectMeshControl/Public/DirectMeshControlSettings.h b/Engine/Plugins/Experimental/Animation/DirectMeshControl/Source/DirectMeshControl/Public/DirectMeshControlSettings.h
--- a/Engine/Plugins/Experimental/Animation/DirectMeshControl/Source/DirectMeshControl/Public/DirectMeshControlSettings.h
+++ b/Engine/Plugins/Experimental/Animation/DirectMeshControl/Source/DirectMeshControl/Public/DirectMeshControlSettings.h
@@ -8,6 +8,9 @@
 
 #define UE_API DIRECTMESHCONTROL_API
 
+/**
+ * UDirectMeshControlSettings are per-project editor settings for the Direct Mesh Control plugin.
+ */
 UCLASS(MinimalAPI, config=EditorPerProjectUserSettings, meta = (Experimental, DisplayName = "Direct Mesh Control Settings"))
 class UDirectMeshControlSettings : public UDeveloperSettings
 {
@@ -15,13 +18,17 @@ class UDirectMeshControlSettings : public UDeveloperSettings
 
 public:
 
-	/** The overlay color used when hovering a direct mesh */
+	/** The overlay color used when hovering a direct mesh. */
 	UPROPERTY(config, EditAnywhere, Category = "Direct Mesh Control Settings")
 	FLinearColor OverlayColor = FLinearColor(0.25f, 0.25f, 0.25f, 0.25f);
 	
+	/** Delegate type broadcast whenever a DMC setting property changes. */
 	DECLARE_MULTICAST_DELEGATE_OneParam(FOnUpdateSettings,  const UDirectMeshControlSettings*);
+	
+	/** Broadcast to all listeners after any property on this settings object is changed. */
 	static UE_API FOnUpdateSettings OnSettingsChange;
 
+	/** UObject override. */
 	UE_API virtual void PostEditChangeProperty(FPropertyChangedEvent& PropertyChangedEvent) override;
 };
 
diff --git a/Engine/Plugins/Experimental/Animation/DirectMeshControl/Source/DirectMeshControl/Public/DirectMeshControlStyle.h b/Engine/Plugins/Experimental/Animation/DirectMeshControl/Source/DirectMeshControl/Public/DirectMeshControlStyle.h
--- a/Engine/Plugins/Experimental/Animation/DirectMeshControl/Source/DirectMeshControl/Public/DirectMeshControlStyle.h
+++ b/Engine/Plugins/Experimental/Animation/DirectMeshControl/Source/DirectMeshControl/Public/DirectMeshControlStyle.h
@@ -6,15 +6,27 @@
 #include "Styling/ISlateStyle.h"
 #include "Styling/SlateStyle.h"
 
+/**
+ * FDirectMeshControlStyle manages the Slate UI style set for the Direct Mesh Control plugin.
+ * Provides static helpers to initialize, shut down, and access the shared FSlateStyleSet used by DMC editor tools and commands.
+ */
+
 class DIRECTMESHCONTROL_API FDirectMeshControlStyle
 {
 public:
+	/** Creates and registers the DMC Slate style set. */
 	static void Initialize();
+
+	/** Unregisters and resets the DMC Slate style set. */
 	static void Shutdown();
 
+	/** Returns the shared DMC Slate style pointer. */
 	static TSharedPtr<class ISlateStyle> Get();
+
+	/** Returns the registered name of the DMC style set. */
 	static FName GetStyleSetName();
 
 private:
+	/** The shared FSlateStyleSet instance owned by this class. */
 	static TSharedPtr<class FSlateStyleSet> StyleSet;
 };
\ No newline at end of file
diff --git a/Engine/Plugins/Experimental/Animation/DirectMeshControl/Source/DirectMeshControl/Public/DirectMeshControlUtilities.h b/Engine/Plugins/Experimental/Animation/DirectMeshControl/Source/DirectMeshControl/Public/DirectMeshControlUtilities.h
--- a/Engine/Plugins/Experimental/Animation/DirectMeshControl/Source/DirectMeshControl/Public/DirectMeshControlUtilities.h
+++ b/Engine/Plugins/Experimental/Animation/DirectMeshControl/Source/DirectMeshControl/Public/DirectMeshControlUtilities.h
@@ -19,7 +19,10 @@ namespace UE::Geometry
 }
 
 /**
- * FGroupSubMeshes
+ * FGroupSubMeshes holds the per-polygroup sub skeletal meshes split from a source skeletal mesh by a triangle label layer.
+ * Each entry in SubSkeletalMeshes element corresponds to one polygroup in the source mesh.
+ * Sub-meshes are generated with a single root-bone binding and a SubToSource vertex attribute that maps sub-mesh vertices back to the
+ * original mesh indices.
  */
 
 USTRUCT()
@@ -27,22 +30,36 @@ struct UE_API FGroupSubMeshes
 {
 	GENERATED_BODY()
 	
+	/**
+	 * Splits the input skeletal mesh into one skeletal mesh per polygroup
+	 * @param InSkeletalMesh  Source skeletal mesh supplying the skeleton.
+	 * @param Mesh            Dynamic mesh whose triangle label layer is read.
+	 * @param LayerName       Name of the triangle label layer that stores polygroup IDs.
+	 */
 	void Rebuild(USkeletalMesh* InSkeletalMesh, const UE::Geometry::FDynamicMesh3* Mesh, const FName LayerName);
+	
+	/** Marks all sub-meshes as garbage and clears the mesh and index-map arrays. */
 	void Reset();
 
+	/** Returns the array of generated sub skeletal meshes. */
 	const TArray<TObjectPtr<USkeletalMesh>>& GetSubSkeletalMeshes() const;
+	
+	/** Returns the map from polygroup ID to index into GetSubSkeletalMeshes(). */
 	const TMap<int32, int32>& GetGroupIdToSubMesh() const;
 	
 private:
 	
+	/** Generated sub skeletal meshes, one entry per unique polygroup in the source layer. */
 	UPROPERTY()
 	TArray<TObjectPtr<USkeletalMesh>> SubSkeletalMeshes;
+	
+	/** Maps each polygroup ID to the corresponding index in SubSkeletalMeshes. */
 	UPROPERTY()
 	TMap<int32, int32> GroupIdToSubMesh;
 };
 
 /**
- * FSubSkeletalMeshData
+ * FSubSkeletalMeshData caches entry associating a (skeletal mesh, layer name) pair with a DDC hash and the resulting per-polygroup sub-meshes.
  */
 
 USTRUCT()
@@ -50,25 +67,36 @@ struct FSubSkeletalMeshData
 {
 	GENERATED_BODY()
 	
+	/** Default constructor (produces an empty, invalid entry). */
 	FSubSkeletalMeshData() = default;
 	
+	/**
+	 * Constructs an entry for the given mesh and layer without yet computing sub-meshes or a DDC hash.
+	 * @param InSkeletalMesh  Source skeletal mesh to track.
+	 * @param InLayerName     Triangle label layer name to track.
+	 */
 	FSubSkeletalMeshData(USkeletalMesh* InSkeletalMesh, const FName InLayerName)
 		: SkeletalMesh(InSkeletalMesh)
 		, LayerName(InLayerName)
 	{}
 
+	/** Soft pointer to the source skeletal mesh. */
 	UPROPERTY()
 	TSoftObjectPtr<USkeletalMesh> SkeletalMesh = nullptr;
 	
+	/** Name of the triangle label layer this entry was built from. */
 	UPROPERTY()
 	FName LayerName = NAME_None;
 	
+	/** DDC hash of the source mesh data; used to detect stale cache entries. */
 	UPROPERTY()
 	uint32 Hash = 0;
 	
+	/** The generated per-polygroup sub-meshes for this (mesh, layer) pair. */
 	UPROPERTY()
 	FGroupSubMeshes GroupSubMeshes;
 
+	/** Hash computation so that FSubSkeletalMeshData can be stored. */
 	friend FORCEINLINE uint32 GetTypeHash(const FSubSkeletalMeshData& Data)
 	{
 		return HashCombine(GetTypeHash(Data.SkeletalMesh), GetTypeHash(Data.LayerName));
@@ -81,16 +109,34 @@ struct FSubSkeletalMeshData
 };
 
 /**
- * 
+ * Utility functions for the Direct Mesh Control system.
+ * Provides accessors for DMC materials, deformers, and cached sub-mesh data, as well as named constants used across the plugin.
  */
 
 namespace UE::DMC
 {
+	/** Returns true if the material-based hover visualization mode is active. */
 	UE_API bool UseMaterialHover();
+	
+	/**
+	 * Fetches or builds the FGroupSubMeshes for the given (mesh, layer) pair.
+	 * @param InSkeletalMesh  Source skeletal mesh.
+	 * @param Mesh            Dynamic mesh representation used for splitting.
+	 * @param LayerName       Triangle label layer that encodes polygroup IDs.
+	 * @return Reference to the cached (or freshly built) group sub-meshes.
+	 */
 	UE_API const FGroupSubMeshes& GetSubMeshes(USkeletalMesh* InSkeletalMesh, const Geometry::FDynamicMesh3* Mesh, const FName LayerName);
+	
+	/** Returns the vertex attribute name used to map sub-mesh vertices back to the source mesh. */
 	UE_API const FName& GetSubToSourceAttrName();
+	
+	/** Returns the Optimus deformer variable name that controls the overlay color. */
 	UE_API const FName& GetColorVarName();
+	
+	/** Lazy-loads and returns the DMC material or nullptr on failure. */
 	UE_API UMaterial* GetMaterial();
+	
+	/** Lazy-loads and returns the DMC deformer or nullptr on failure. */
 	UE_API UOptimusDeformer* GetDeformer();
 }
 	
diff --git a/Engine/Plugins/Experimental/Animation/DirectMeshControl/Source/DirectMeshControl/Public/Tools/DirectMeshControlTool.h b/Engine/Plugins/Experimental/Animation/DirectMeshControl/Source/DirectMeshControl/Public/Tools/DirectMeshControlTool.h
--- a/Engine/Plugins/Experimental/Animation/DirectMeshControl/Source/DirectMeshControl/Public/Tools/DirectMeshControlTool.h
+++ b/Engine/Plugins/Experimental/Animation/DirectMeshControl/Source/DirectMeshControl/Public/Tools/DirectMeshControlTool.h
@@ -29,7 +29,7 @@ class UOptimusDeformer;
 using UE::Geometry::FDynamicMesh3;
 
 /**
- * UDirectMeshControlToolBuilder
+ * UDirectMeshControlToolBuilder builds a new UDirectMeshControlTool instance 
  */
 
 UCLASS()
@@ -38,19 +38,27 @@ class DIRECTMESHCONTROL_API UDirectMeshControlToolBuilder : public UInteractiveT
 	GENERATED_BODY()
 
 public:
-	// UInteractiveToolBuilder overrides
+	/**
+	 * Returns @c true when the scene state contains a single selectable primitive component that the tool can operate on.
+	 * @param SceneState  Current tool-manager scene state.
+	 * @return Whether the tool can be built for the given selection.
+	 */
 	virtual bool CanBuildTool(const FToolBuilderState& SceneState) const override;
+	
+	/**
+	 * Instantiates and initializes a new UDirectMeshControlTool for the given scene state.
+	 * @param SceneState  Current tool-manager scene state.
+	 * @return Newly created tool instance.
+	 */
 	virtual UInteractiveTool* BuildTool(const FToolBuilderState& SceneState) const override;
-
-	TFunction<void (const USkeletalMeshComponent*, const FName)> PreClickFunction;
 	
 protected:
-	// UInteractiveToolWithToolTargetsBuilder overrides
+	/** Returns the target type requirements for this builder (a single primitive component target). */
 	virtual const FToolTargetTypeRequirements& GetTargetRequirements() const override;
 };
 
 /**
- * UDirectMeshControlTool
+ * UDirectMeshControlTool visualizes and allows posing of per-polygroup sub-meshes on a skeletal mesh.
  */
 
 UCLASS()
@@ -64,6 +72,10 @@ class DIRECTMESHCONTROL_API UDirectMeshControlTool :
 
 public:
 
+	/**
+	 * Stores the data needed during setup. Must be called by the builder before Setup().
+	 * @param InSceneState  Scene state captured at tool-build time.
+	 */
 	void Init(const FToolBuilderState& InSceneState);
 	
 	// UInteractiveTool overrides
@@ -97,6 +109,7 @@ public:
 	// IInteractiveToolCameraFocusAPI overrides
 	virtual FBox GetWorldSpaceFocusBox() override;
 
+	/** When true, the tool will rebuild topology on the next tick. */
 	bool bRebuild = true;
 	
 protected:
@@ -109,125 +122,195 @@ protected:
 
 public:
 	
+	/** Editable display, offset, and action-mode properties shown in the Details panel. */
 	UPROPERTY()
 	TObjectPtr<UDirectMeshControlProperties> Properties;
 
+	/** Polygon selection mechanic used in Select mode to pick polygroups / bones. */
 	UPROPERTY()
 	TObjectPtr<UPolygonSelectionMechanic> SelectionMechanic = nullptr;
 
+	/** Currently selected bone names. */
 	TArray<FName> Selection;
+	
+	/** Returns the name of the primary selected bone or NAME_None if nothing is selected. */
 	FName GetCurrentBone() const;
 
+	/** Maps polygroup ID to the index of the corresponding bone in the source skeleton. */
 	TMap<int32, int32> PolyGroupToBone;
+	
+	/** Maps polygroup ID to the index of the corresponding sub-mesh in SubComponents. */
 	TMap<int32, int32> PolyGroupToSubMesh;
 	
 protected:
 
+	/**
+	 * Casts a ray against the preview mesh and returns the polygroup ID of the hit triangle, or INDEX_NONE on miss.
+	 * @param InRay  Device ray to cast.
+	 * @return Polygroup ID of the hit triangle, or INDEX_NONE.
+	 */
 	int32 GetHoveredPolyGroup(const FInputDeviceRay& InRay) const;
 	
+	/** Single-click behavior registered for bone/polygroup selection. */
 	UPROPERTY()
 	TObjectPtr<USingleClickInputBehavior> LeftClickBehavior = nullptr;
 
+	/** Preview mesh component that displays the source skeletal mesh with polygroup coloring. */
 	UPROPERTY()
 	TObjectPtr<UPreviewMesh> PreviewMesh = nullptr;
 
+	/** World in which the tool is operating; used for spawning sub-mesh components. */
 	UPROPERTY()
 	TWeakObjectPtr<UWorld> TargetWorld = nullptr;
 
+	/** Skeletal mesh editor context, used to integrate with the Persona editor. */
 	UPROPERTY()
 	TWeakObjectPtr<USkeletalMeshEditorContextObjectBase> EditorContext = nullptr;
+	
+	/** Persona mode manager context, used to query and modify the active bone selection. */	
 	UPROPERTY()
 	TWeakObjectPtr<UPersonaEditorModeManagerContext> PersonaModeManagerContext = nullptr;
 
+	/** Transient UDirectMeshControlComponent instances spawned for each polygroup sub-mesh. */
 	TArray<TWeakObjectPtr<USkeletalMeshComponent>> SubComponents;
 	
 private:
+	
+	/** Returns the bone indices in the source skeleton that correspond to the currently selected bone names in Selection.
+	 */
 	TArray<int32> GetSelectedBoneIndexes() const;
 
+	/** AABB tree for the preview mesh, used for polygroup hover hit-testing. */
 	TSharedPtr<UE::Geometry::FDynamicMeshAABBTree3> MeshSpatial;
+
+	/** Returns a reference to the spatial AABB tree, rebuilding it if bSpatialDirty is true. */
 	UE::Geometry::FDynamicMeshAABBTree3& GetSpatial();
+
+	/** true when the preview mesh geometry has changed and MeshSpatial must be rebuilt. */
 	bool bSpatialDirty = true;
 
+	/** Polygroup ID currently under the cursor, or INDEX_NONE. */
 	int32 CurrentPolyGroup = INDEX_NONE;
+
+	/** Bone index currently under the cursor (derived from CurrentPolyGroup), or INDEX_NONE. */
 	int32 CurrentBone = INDEX_NONE;
-	
+
+	/** Last screen-space cursor position used to avoid redundant spatial queries. */
 	mutable FVector2D LastScreenPosition = FVector2D::ZeroVector;
+
+	/** Triangle index returned by the last spatial hit test (cached to avoid redundant queries). */
 	mutable int32 LastHit = INDEX_NONE;
 
-	// setup
+	/** Registers properties and wires property-change callbacks. */
 	void SetupProperties();
+
+	/** Creates the preview mesh. */
 	void SetupPreviewMesh();
+
+	/** Registers click and hover input behaviors with the tool behavior manager. */
 	void SetupBehaviors();
+
+	/** Initializes the selection mechanic with the current polygroup topology. */
 	void SetupPolyGroupSelection();
 
+	/** Generates sub-skeletal mesh instances for each polygroup sub-mesh. */
 	void GenerateDMCSubSkeletalMeshes();
+
+	/** Destroys all sub-skeletal mesh instances. */
 	void CleanDMCSubSkeletalMeshes();
 
+	/** Applies current properties settings to all sub-mesh components. */
 	void UpdateVisualization();
-	
-	TFunction<void()> PendingFunction;
-	TFunction<void()> ReleasePendingFunction;
 
+	/** Group topology built from the active polygroup set (used by the selection mechanic). */
 	TUniquePtr<UE::Geometry::FGroupTopology> Topology = nullptr;
 
+	/** Weak pointer to the source skeletal mesh component being edited. */
 	TWeakObjectPtr<USkeletalMeshComponent> SourceSkeletalMeshComponent = nullptr;
+
+	/** Weak pointer to the source skeletal mesh asset. */
 	TWeakObjectPtr<USkeletalMesh> SourceSkeletalMesh = nullptr;
 
+	/** Active polygroup set derived from the triangle label layer of the source mesh. */
 	TUniquePtr<UE::Geometry::FPolygroupSet> ActiveGroupSet;
+
+	/** Dynamic mesh polygroup attribute backing ActiveGroupSet. */
 	TUniquePtr<UE::Geometry::FDynamicMeshPolygroupAttribute> PolygroupAttrib;
 
 	friend class UDirectMeshControlProperties;
 };
 
 /**
- * UDirectMeshControlProperties
+ * EDMCAction controls which action is performed when the user clicks or hovers over a polygroup in the viewport.
  */
-
 UENUM()
 enum class EDMCAction
 {
+	/** No interaction, clicks and hover have no effect. */
 	None,
+	/** Click selects the polygroup's associated bone via the selection mechanic. */
 	Select,
+	/** Click assigns the hovered polygroup to the currently selected bone. */
 	Assign,
+	/** Hover highlights the polygroup, click applies a positional offset to the sub-mesh. */
 	Pose
 };
 
+/**
+ * UDirectMeshControlProperties represents editable properties for UDirectMeshControlTool.
+ */
+
 UCLASS()
 class DIRECTMESHCONTROL_API UDirectMeshControlProperties : public UInteractiveToolPropertySet
 {
 	GENERATED_BODY()
 public:
+	/**
+	 * Stores a weak reference to InParentTool so that button-action functions can forward calls to the tool.
+	 * @param InParentTool  The @c UDirectMeshControlTool that owns these properties.
+	 */
 	void Initialize(UDirectMeshControlTool* InParentTool);
 
+	/** Opacity blend factor applied to all polygroup sub-mesh overlays. Range [0, 1]. */
 	UPROPERTY(EditAnywhere, Category = Display, meta = (ClampMin = "0.0", ClampMax = "1.0", UIMin = "0.0", UIMax = "1.0"))
 	float Alpha = 0.5f;
 
+	/** World-space offset applied to the entire preview mesh visualization. */
 	UPROPERTY(EditAnywhere, Category = Display)
 	FVector Offset;
 
+	/** Additional world-space offset applied to each individual polygroup sub-mesh component. */
 	UPROPERTY(EditAnywhere, Category = Display)
 	FVector OffsetSub;
 
+	/** When true, the individual per-polygroup sub-mesh components are visible. */
 	UPROPERTY(EditAnywhere, Category = Display)
 	bool bSubMeshes = true;
-	
+
+	/** When true, sub-mesh components are attached to the target actor's root component. */
 	UPROPERTY(EditAnywhere, Category = Display)
 	bool bAttachToActor = true;
-	
+
+	/** When true, sub-mesh components respond to hit-proxy selection in the viewport. */
 	UPROPERTY(EditAnywhere, Category = Display)
 	bool bHitProxy = true;
-	
+
+	/** When true, hover detection is performed through the material rather than via hit proxies. */
 	UPROPERTY(EditAnywhere, Category = Display)
 	bool bHoverThruMaterial = true;
-	
+
+	/** Current interaction mode; determines how clicks and hover events are handled. */
 	UPROPERTY(EditAnywhere, Category = Action)
 	EDMCAction Action = EDMCAction::Pose;
 
+	/** Assigns the currently hovered polygroup to the selected bone (callable from the Details panel). */
 	UFUNCTION(CallInEditor, Category = Action)
 	void Assign();
-	
+
+	/** Runs the default bone-to-polygroup assignment heuristic (closest bone per polygroup centroid). */
 	void DefaultAssign();
-	
+
 private:
+	/** Weak back-pointer to the owning tool. */
 	TWeakObjectPtr<UDirectMeshControlTool> ParentTool = nullptr;
 };
diff --git a/Engine/Plugins/Experimental/Animation/DirectMeshControl/Source/DirectMeshControl/Public/Tools/DirectMeshPolygroupTool.h b/Engine/Plugins/Experimental/Animation/DirectMeshControl/Source/DirectMeshControl/Public/Tools/DirectMeshPolygroupTool.h
--- a/Engine/Plugins/Experimental/Animation/DirectMeshControl/Source/DirectMeshControl/Public/Tools/DirectMeshPolygroupTool.h
+++ b/Engine/Plugins/Experimental/Animation/DirectMeshControl/Source/DirectMeshControl/Public/Tools/DirectMeshPolygroupTool.h
@@ -22,18 +22,35 @@ namespace UE::Geometry
 }
 
 /**
- *
+ * UDirectMeshPolygroupToolBuilder builds a UDirectMeshPolygroupTool instance for the given scene state.
  */
 UCLASS()
 class DIRECTMESHCONTROL_API UDirectMeshPolygroupToolBuilder : public USingleTargetWithSelectionToolBuilder
 {
 	GENERATED_BODY()
 public:
+	/**
+	 * Creates and returns a new UDirectMeshPolygroupTool instance for the provided scene state.
+	 * @param SceneState  Current tool-manager scene state.
+	 * @return Newly created tool instance.
+	 */
 	virtual USingleTargetWithSelectionTool* CreateNewTool(const FToolBuilderState & SceneState) const override;
+
+	/** Returns false, the polygroup tool does not require an active selection to launch. */
 	virtual bool RequiresInputSelection() const override { return false; }
+
+	/**
+	 * Returns true when the scene state contains exactly one non-volume skeletal mesh component.
+	 * @param SceneState  Current tool-manager scene state.
+	 * @return Whether the tool can be built for the given selection.
+	 */
 	virtual bool CanBuildTool(const FToolBuilderState& SceneState) const override;
 };
 
+/**
+ * UDirectMeshPolygroupToolProperties represents editable properties controlling polygroup generation behavior and display.
+ */
+
 UCLASS()
 class DIRECTMESHCONTROL_API UDirectMeshPolygroupToolProperties : public UInteractiveToolPropertySet
 {
@@ -82,24 +99,28 @@ public:
 	FName LayerName = "dmc-polygroup";
 };
 
-
-
-
+/**
+ * Operator factory for the Direct Mesh Polygroup tool.
+ */
 UCLASS()
 class DIRECTMESHCONTROL_API UDirectMeshPolygroupOperatorFactory : public UObject, public UE::Geometry::IDynamicMeshOperatorFactory
 {
 	GENERATED_BODY()
 
 public:
-	// IDynamicMeshOperatorFactory API
+	/**
+	 * Creates a new FDirectMeshPolygroupOp operator populated with the current settings from DirectMeshPolygroupTool.
+	 * @return Unique pointer to the configured operator ready for background execution.
+	 */
 	virtual TUniquePtr<UE::Geometry::FDynamicMeshOperator> MakeNewOperator() override;
 
+	/** Back pointer to the owning tool, provides settings and mesh data for operator creation. */
 	UPROPERTY()
-	TObjectPtr<UDirectMeshPolygroupTool> DirectMeshPolygroupTool;  // back pointer
+	TObjectPtr<UDirectMeshPolygroupTool> DirectMeshPolygroupTool;
 };
 
 /**
- *
+ * UDirectMeshPolygroupTool generates bone-based polygroups on a skeletal mesh asynchronously and commits the result to a named triangle label layer on Accept.
  */
 UCLASS()
 class DIRECTMESHCONTROL_API UDirectMeshPolygroupTool : public USingleTargetWithSelectionTool
@@ -109,52 +130,75 @@ class DIRECTMESHCONTROL_API UDirectMeshPolygroupTool : public USingleTargetWithS
 public:
 	UDirectMeshPolygroupTool();
 
+	/** Builds the preview mesh, registers properties, and starts the background compute operator. */
 	virtual void Setup() override;
+
+	/**
+	 * Commits the computed polygroup layer to the skeletal mesh on Accept, or discards the result on Cancel.
+	 * @param ShutdownType  Whether the tool is being accepted or canceled.
+	 */
 	virtual void OnShutdown(EToolShutdownType ShutdownType) override;
 
+	/**
+	 * Ticks the background compute operator and refreshes the visualization when a new result is available.
+	 * @param DeltaTime  Time elapsed since the last tick.
+	 */
 	virtual void OnTick(float DeltaTime) override;
 
+	/** Returns true, the polygroup tool supports cancellation. */
 	virtual bool HasCancel() const override { return true; }
+
+	/** Returns true, the polygroup tool supports acceptance. */
 	virtual bool HasAccept() const override { return true; }
+
+	/** Returns true when the background operator has produced a valid result that can be committed. */
 	virtual bool CanAccept() const override;
 
-	// update parameters in DirectMeshPolygroupOp based on current Settings
+	/**
+	 * Copies the current properties values and bone data into DirectMeshPolygroupOp so that the operator uses up-to-date settings on the next compute pass.
+	 * @param DirectMeshPolygroupOp  Operator instance to populate.
+	 */
 	void UpdateOpParameters(FDirectMeshPolygroupOp& DirectMeshPolygroupOp) const;
 
 protected:
-	
+
+	/** Invalidates the preview compute, triggering a new background operator run. */
 	void OnSettingsModified();
 
-protected:
+	/** Editable properties displayed in the Details panel. */
 	UPROPERTY()
 	TObjectPtr<UDirectMeshPolygroupToolProperties> Properties;
 
+	/** Background compute manager that runs FDirectMeshPolygroupOp asynchronously. */
 	UPROPERTY()
 	TObjectPtr<UMeshOpPreviewWithBackgroundCompute> PreviewCompute = nullptr;
 
+	/** Geometry used to render polygroup border edges on the preview mesh. */
 	UPROPERTY()
 	TObjectPtr<UPreviewGeometry> PreviewGeometry = nullptr;
 
-	// If a selection was provided (bUsingSelection = true), UnmodifiedAreaPreviewMesh is used to render the unmodified (non-selected) part of the input mesh,
-	// as the PreviewCompute input mesh will be limited to the input selected area
+	/** Preview mesh that renders the non-selected (unmodified) portion of the source mesh when a selection ROI is active. */
 	UPROPERTY()
 	TObjectPtr<UPreviewMesh> UnmodifiedAreaPreviewMesh = nullptr;
 
-protected:
-	// copy of the input mesh (possibly does not need to be SharedPtr anymore now that FSharedConstDynamicMesh3 is used below, but then needs PimplPtr...)
+	/** Thread-safe copy of the full input dynamic mesh, used as the base for all compute passes. */
 	TSharedPtr<UE::Geometry::FDynamicMesh3, ESPMode::ThreadSafe> OriginalDynamicMesh;
 
-	// If there is an active selection, bUsingSelection will be true and the ROI and Submesh will be initialized
+	/** true when the tool was launched with an active triangle selection. */
 	bool bUsingSelection = false;
+
+	/** Triangle indices that form the selection ROI; populated only when bUsingSelection is true. */ 
 	TSharedPtr<TSet<int32>> SelectionTriangleROI;
+
+	/** Sub-mesh extracted from OriginalDynamicMesh for the selection ROI, populated only when bUsingSelection is true. */
 	TSharedPtr<UE::Geometry::FDynamicSubmesh3, ESPMode::ThreadSafe> OriginalSubmesh;
 
-	// The mesh passed to the Compute operator to base result on, will either be OriginalDynamicMesh or the submesh from OriginalSubmesh if using a selection
+	/** Shared mesh reference passed to the compute operator, points to the full mesh or the selection sub-mesh depending on bUsingSelection. */
 	TSharedPtr<UE::Geometry::FSharedConstDynamicMesh3> ComputeOperatorSharedMesh;
 
-	// current set of detected polygroup edges, relative to ComputeOperatorSharedMesh
+	/** Edge indices (relative to ComputeOperatorSharedMesh) that form the detected polygroup boundaries, used for edge visualization. */
 	TArray<int> PolygonEdges;
-	
-	void UpdateVisualization();
 
+	/** Refreshes the polygroup edge overlay and per-group color visualization after a new compute result arrives. */
+	void UpdateVisualization();
 };
diff --git a/Engine/Plugins/Experimental/Animation/DirectMeshControl/Source/DirectMeshControlRig/Public/DirectMeshControlProxyWrapper.h b/Engine/Plugins/Experimental/Animation/DirectMeshControl/Source/DirectMeshControlRig/Public/DirectMeshControlProxyWrapper.h
--- a/Engine/Plugins/Experimental/Animation/DirectMeshControl/Source/DirectMeshControlRig/Public/DirectMeshControlProxyWrapper.h
+++ b/Engine/Plugins/Experimental/Animation/DirectMeshControl/Source/DirectMeshControlRig/Public/DirectMeshControlProxyWrapper.h
@@ -7,7 +7,8 @@
 #include "DirectMeshControlProxyWrapper.generated.h"
 
 /**
- * 
+ * UDirectMeshControlProxy is a transient proxy object used by FShapeProxyComponentProviderRegistry
+ * to associate a UDirectMeshControlComponent with a Control Rig shape.
  */
 UCLASS(transient)
 class DIRECTMESHCONTROLRIG_API UDirectMeshControlProxy : public UObject
diff --git a/Engine/Plugins/Experimental/Animation/DirectMeshControl/Source/DirectMeshControlRig/Public/DirectMeshControlRigModule.h b/Engine/Plugins/Experimental/Animation/DirectMeshControl/Source/DirectMeshControlRig/Public/DirectMeshControlRigModule.h
--- a/Engine/Plugins/Experimental/Animation/DirectMeshControl/Source/DirectMeshControlRig/Public/DirectMeshControlRigModule.h
+++ b/Engine/Plugins/Experimental/Animation/DirectMeshControl/Source/DirectMeshControlRig/Public/DirectMeshControlRigModule.h
@@ -4,14 +4,17 @@
 
 #include "Modules/ModuleManager.h"
 
+/**
+ * FDirectMeshControlRigModule is an editor module that integrates the Direct Mesh Control plugin with the Control Rig system.
+ */
 class FDirectMeshControlRigModule : public IModuleInterface
 {
 public:
-
 	/** IModuleInterface implementation */
 	virtual void StartupModule() override;
 	virtual void ShutdownModule() override;
-	
+
 private:
+	/** Called after the engine has finished initializing. */
 	void OnPostEngineInit() const;
 };
diff --git a/Engine/Plugins/Experimental/Animation/DirectMeshControl/Source/DirectMeshControlRig/Public/Units/RigUnit_DirectMeshControl.h b/Engine/Plugins/Experimental/Animation/DirectMeshControl/Source/DirectMeshControlRig/Public/Units/RigUnit_DirectMeshControl.h
--- a/Engine/Plugins/Experimental/Animation/DirectMeshControl/Source/DirectMeshControlRig/Public/Units/RigUnit_DirectMeshControl.h
+++ b/Engine/Plugins/Experimental/Animation/DirectMeshControl/Source/DirectMeshControlRig/Public/Units/RigUnit_DirectMeshControl.h
@@ -8,8 +8,10 @@
 
 
 /**
-* DMC FRigUnit_SetupShapeLibraryFromLayer
-*/
+ * FRigUnit_SetupShapeLibraryFromLayer reads a polygroup triangle label layer from the source skeletal mesh, generates per-polygroup sub-meshes
+ * and registers a UControlRigShapeLibrary out of it.
+ */
+
 USTRUCT(meta=(DisplayName="Set Shape Library from Layer", Keywords="Direct Mesh Control, Surface Selection", Varying))
 struct FRigUnit_SetupShapeLibraryFromLayer : public FRigUnit_DynamicHierarchyBaseMutable
 {
@@ -23,15 +25,11 @@ struct FRigUnit_SetupShapeLibraryFromLayer : public FRigUnit_DynamicHierarchyBas
 	RIGVM_METHOD()
 	virtual void Execute() override;
 
-	/*
-	 * The layer to extract
-	 */
+	/* The layer to extract */
 	UPROPERTY(meta = (Input))
 	FName LayerName;
 	
-	/*
-	 * The extracted groups
-	 */
+	/* The extracted groups */
 	UPROPERTY(meta = (Output))
 	TArray<FName> GroupNames;
 };
