Animation Mode: support for hovered state and colors for gizmo libraries

#rnx

#Synced-CL 52777971

diff --git a/Engine/Plugins/Animation/ControlRig/Source/ControlRig/Public/ControlRigGizmoLibrary.h b/Engine/Plugins/Animation/ControlRig/Source/ControlRig/Public/ControlRigGizmoLibrary.h
--- a/Engine/Plugins/Animation/ControlRig/Source/ControlRig/Public/ControlRigGizmoLibrary.h
+++ b/Engine/Plugins/Animation/ControlRig/Source/ControlRig/Public/ControlRigGizmoLibrary.h
@@ -75,6 +75,12 @@ public:
 	UPROPERTY(EditAnywhere, Category = "ShapeLibrary")
 	FName MaterialColorParameter;
 
+	UPROPERTY(EditAnywhere, Category = "ShapeLibrary")
+	FName MaterialHoveredParameter = NAME_None;
+	
+	UPROPERTY(EditAnywhere, Category = "ShapeLibrary")
+	FName MaterialHoveredColorParameter = NAME_None;
+	
 	UPROPERTY(EditAnywhere, Category = "ShapeLibrary")
 	TArray<FControlRigShapeDefinition> Shapes;
 
diff --git a/Engine/Plugins/Animation/ControlRig/Source/ControlRigEditor/Private/EditMode/ControlRigEditMode.cpp b/Engine/Plugins/Animation/ControlRig/Source/ControlRigEditor/Private/EditMode/ControlRigEditMode.cpp
--- a/Engine/Plugins/Animation/ControlRig/Source/ControlRigEditor/Private/EditMode/ControlRigEditMode.cpp
+++ b/Engine/Plugins/Animation/ControlRig/Source/ControlRigEditor/Private/EditMode/ControlRigEditMode.cpp
@@ -525,6 +525,10 @@ void FControlRigEditMode::Enter()
 		{
 			Settings->OnSettingsChange.AddSP(this, &FControlRigEditMode::OnSettingsChanged);
 		}
+		if (!Settings->OnPropertyChanged.IsBoundToObject(this))
+		{
+			Settings->OnPropertyChanged.AddSP(this, &FControlRigEditMode::OnSettingsPropertyChanged);
+		}
 	}
 
 	if (IsInLevelEditor())
@@ -682,6 +686,7 @@ void FControlRigEditMode::Exit()
 	if (const UControlRigEditModeSettings* Settings = GetSettings())
 	{
 		Settings->OnSettingsChange.RemoveAll(this);
+		Settings->OnPropertyChanged.RemoveAll(this);
 	}
 	
 	// Call parent implementation
@@ -4811,6 +4816,7 @@ void FControlRigEditMode::CreateShapeActors(UControlRig* InControlRig)
 		GetMutableDefault<UEditorPerProjectUserSettings>()->bAllowSelectTranslucent = true;
 	}
 
+	const UControlRigEditModeSettings* Settings = GetSettings();
 	const TArray<TSoftObjectPtr<UControlRigShapeLibrary>> ShapeLibraries = InControlRig->GetShapeLibraries();
 
 	const int32 ControlRigIndex = RuntimeControlRigs.Find(InControlRig);
@@ -4858,6 +4864,12 @@ void FControlRigEditMode::CreateShapeActors(UControlRig* InControlRig)
 				}
 			}
 			Param.ColorParameterName = ShapeDef->Library->MaterialColorParameter;
+			Param.HoveredParameterName = ShapeDef->Library->MaterialHoveredParameter;
+			Param.HoveredColorParameterName = ShapeDef->Library->MaterialHoveredColorParameter;
+			if (Settings)
+			{
+				Param.HoveredColor = Settings->HoveredColor;
+			}
 			Param.PostSetupFunction = ShapeDef->PostSetupFunction;
 		}
 
@@ -5520,6 +5532,34 @@ void FControlRigEditMode::OnSettingsChanged(const UControlRigEditModeSettings* I
 	}
 }
 
+void FControlRigEditMode::OnSettingsPropertyChanged(const UControlRigEditModeSettings* InSettings, const FPropertyChangedEvent& PropertyChangedEvent)
+{
+	if (!InSettings)
+	{
+		return;
+	}
+	
+	if (PropertyChangedEvent.GetPropertyName() == GET_MEMBER_NAME_CHECKED(UControlRigEditModeSettings, HoveredColor))
+	{
+		if (RecreateControlShapesRequired == ERecreateControlRigShape::RecreateNone)
+		{
+			for (const auto&[ControlRig, ShapeActors]: ControlRigShapeActors)
+			{
+				if (ControlRig.IsValid())
+				{
+					for (AControlRigShapeActor* ShapeActor : ShapeActors)
+					{
+						if (ShapeActor)
+						{
+							ShapeActor->SetHoveredColor(InSettings->HoveredColor);
+						}
+					}
+				}
+			}
+		}
+	}
+}
+
 void FControlRigEditMode::SetControlShapeTransform(
 	const AControlRigShapeActor* InShapeActor,
 	const FTransform& InGlobalTransform,
diff --git a/Engine/Plugins/Animation/ControlRig/Source/ControlRigEditor/Private/EditMode/ControlRigEditModeSettings.cpp b/Engine/Plugins/Animation/ControlRig/Source/ControlRigEditor/Private/EditMode/ControlRigEditModeSettings.cpp
--- a/Engine/Plugins/Animation/ControlRig/Source/ControlRigEditor/Private/EditMode/ControlRigEditModeSettings.cpp
+++ b/Engine/Plugins/Animation/ControlRig/Source/ControlRigEditor/Private/EditMode/ControlRigEditModeSettings.cpp
@@ -7,6 +7,7 @@
 #include UE_INLINE_GENERATED_CPP_BY_NAME(ControlRigEditModeSettings)
 
 UControlRigEditModeSettings::FOnUpdateSettings UControlRigEditModeSettings::OnSettingsChange;
+UControlRigEditModeSettings::FOnPropertyChanged UControlRigEditModeSettings::OnPropertyChanged;
 
 void UControlRigEditModeSettings::PreEditChange(FProperty* PropertyAboutToChange)
 {
@@ -28,6 +29,7 @@ void UControlRigEditModeSettings::PostEditChangeProperty(struct FPropertyChanged
 		SaveConfig();
 	}
 	OnSettingsChange.Broadcast(this);
+	OnPropertyChanged.Broadcast(this, PropertyChangedEvent);
 }
 
 #if WITH_EDITOR
@@ -35,6 +37,7 @@ void UControlRigEditModeSettings::PostEditUndo()
 {
 	GizmoScaleDelegate.Broadcast(GizmoScale);
 	OnSettingsChange.Broadcast(this);
+	OnPropertyChanged.Broadcast(this, nullptr);
 }
 #endif
 
diff --git a/Engine/Plugins/Animation/ControlRig/Source/ControlRigEditor/Private/EditMode/ControlRigEditModeSettings.h b/Engine/Plugins/Animation/ControlRig/Source/ControlRigEditor/Private/EditMode/ControlRigEditModeSettings.h
--- a/Engine/Plugins/Animation/ControlRig/Source/ControlRigEditor/Private/EditMode/ControlRigEditModeSettings.h
+++ b/Engine/Plugins/Animation/ControlRig/Source/ControlRigEditor/Private/EditMode/ControlRigEditModeSettings.h
@@ -134,9 +134,16 @@ public:
 	UPROPERTY(config, EditAnywhere, Category = "Animation Settings|Interaction", meta=(Experimental))
 	EClickAndDragBehavior ClickAndDragBehavior = EClickAndDragBehavior::Disabled;
 	
+	/** Hovered color used when hovering a control shape. */
+    UPROPERTY(config, EditAnywhere, Category = "Animation Settings", meta = (EditCondition = "!bHideControlShapes"))
+    FLinearColor HoveredColor = FLinearColor(0.5f, 0.5f, 0.5f, 0.25f);
+	
 	/** Delegate broadcasted whenever GizmoScale is modified */
 	FGizmoScaleSet GizmoScaleDelegate;
 
 	DECLARE_MULTICAST_DELEGATE_OneParam(FOnUpdateSettings, const UControlRigEditModeSettings*);
 	static FOnUpdateSettings OnSettingsChange;
+	
+	DECLARE_MULTICAST_DELEGATE_TwoParams(FOnPropertyChanged, const UControlRigEditModeSettings*, const FPropertyChangedEvent&);
+	static FOnPropertyChanged OnPropertyChanged;
 };
\ No newline at end of file
diff --git a/Engine/Plugins/Animation/ControlRig/Source/ControlRigEditor/Public/EditMode/ControlRigEditMode.h b/Engine/Plugins/Animation/ControlRig/Source/ControlRigEditor/Public/EditMode/ControlRigEditMode.h
--- a/Engine/Plugins/Animation/ControlRig/Source/ControlRigEditor/Public/EditMode/ControlRigEditMode.h
+++ b/Engine/Plugins/Animation/ControlRig/Source/ControlRigEditor/Public/EditMode/ControlRigEditMode.h
@@ -607,6 +607,7 @@ private:
 	UE_API bool CanChangeControlShapeTransform();
 
 	UE_API void OnSettingsChanged(const UControlRigEditModeSettings* InSettings);
+	UE_API void OnSettingsPropertyChanged(const UControlRigEditModeSettings* InSettings, const FPropertyChangedEvent& PropertyChangedEvent);
 	
 public:
 	//Toolbar functions
diff --git a/Engine/Plugins/Experimental/Animation/DirectMeshControl/Source/DirectMeshControl/Private/DirectMeshControlSettings.cpp b/Engine/Plugins/Experimental/Animation/DirectMeshControl/Source/DirectMeshControl/Private/DirectMeshControlSettings.cpp
deleted file mode 100644
diff --git a/Engine/Plugins/Experimental/Animation/DirectMeshControl/Source/DirectMeshControl/Public/DirectMeshControlSettings.h b/Engine/Plugins/Experimental/Animation/DirectMeshControl/Source/DirectMeshControl/Public/DirectMeshControlSettings.h
deleted file mode 100644
diff --git a/Engine/Plugins/Experimental/Animation/DirectMeshControl/Source/DirectMeshControlRig/Private/Units/RigUnit_DirectMeshControl.cpp b/Engine/Plugins/Experimental/Animation/DirectMeshControl/Source/DirectMeshControlRig/Private/Units/RigUnit_DirectMeshControl.cpp
--- a/Engine/Plugins/Experimental/Animation/DirectMeshControl/Source/DirectMeshControlRig/Private/Units/RigUnit_DirectMeshControl.cpp
+++ b/Engine/Plugins/Experimental/Animation/DirectMeshControl/Source/DirectMeshControlRig/Private/Units/RigUnit_DirectMeshControl.cpp
@@ -6,7 +6,6 @@
 #include "ControlRig.h"
 #include "ControlRigGizmoActor.h"
 #include "DirectMeshControlProxyWrapper.h"
-#include "DirectMeshControlSettings.h"
 #include "DirectMeshControlUtilities.h"
 #include "IControlRigObjectBinding.h"
 #include "MeshDescriptionToDynamicMesh.h"
@@ -128,6 +127,8 @@ FRigUnit_SetupShapeLibraryFromLayer_Execute()
 				NewDMCShapeLibrary->Shapes.Reserve(GroupNames.Num());
 				NewDMCShapeLibrary->DefaultMaterial = UE::DMC::GetMaterial();
 				NewDMCShapeLibrary->XRayMaterial = UE::DMC::GetMaterial();
+				NewDMCShapeLibrary->MaterialHoveredParameter = TEXT("Hovered");
+				NewDMCShapeLibrary->MaterialHoveredColorParameter = TEXT("HoveredColor");
 
 				const FGroupSubMeshes& GroupSubMeshes = UE::DMC::GetSubMeshes(SkeletalMesh, &DynaMesh, LayerName);
 				const TArray<TObjectPtr<USkeletalMesh>>& SubMeshes = GroupSubMeshes.GetSubSkeletalMeshes();
@@ -139,9 +140,6 @@ FRigUnit_SetupShapeLibraryFromLayer_Execute()
 				},
 				FNameLexicalLess());
 				
-				const UDirectMeshControlSettings* DirectMeshControlSettings = GetMutableDefault<UDirectMeshControlSettings>();
-				constexpr FLinearColor DefaultOverlay(0.25f, 0.25f, 0.25f, 0.25f);
-				const FLinearColor& OverlayColor = DirectMeshControlSettings ? DirectMeshControlSettings->OverlayColor : DefaultOverlay;
 				for (int32 Index = 0; Index < SubMeshes.Num(); ++Index)
 				{
 					const auto& [GroudID, SubMeshIndex] = GroupIdToSubMesh[Index];
@@ -155,7 +153,7 @@ FRigUnit_SetupShapeLibraryFromLayer_Execute()
 						ShapeDefinition.ShapeName = LabelAdapter.GetValueFromGroup(GroudID);
 						// ShapeDefinition.ShapeProxy = SubSkeletalMesh;
 						ShapeDefinition.ShapeProxy = NewObject<UDirectMeshControlProxy>(SubSkeletalMesh);
-						ShapeDefinition.PostSetupFunction = [GroudID, OverlayColor](AControlRigShapeActor* InShape)
+						ShapeDefinition.PostSetupFunction = [GroudID](AControlRigShapeActor* InShape)
 						{
 							if (USkeletalMeshComponent* DMCComponent = InShape ? Cast<USkeletalMeshComponent>(InShape->ProxyComponent.Get()) : nullptr)
 							{
@@ -181,18 +179,7 @@ FRigUnit_SetupShapeLibraryFromLayer_Execute()
 								DMCComponent->SetMeshDeformer(UE::DMC::GetDeformer());
 								
 								// setup material
-								if (UE::DMC::UseMaterialHover())
-								{
-									if (UMaterialInstanceDynamic* MaterialInstance = Cast<UMaterialInstanceDynamic>(DMCComponent->GetMaterial(0)))
-									{
-										InShape->HoveredParameterName = TEXT("Hovered");
-										MaterialInstance->SetScalarParameterValue(TEXT("Hovered"), 0.f);
-										
-										InShape->HoveredColorParameterName = TEXT("HoveredColor");
-										InShape->SetHoveredColor(OverlayColor);
-									}
-								}
-								else
+								if (!UE::DMC::UseMaterialHover())
 								{
 									UOptimusDeformerInstance* MeshDeformerInstance = Cast<UOptimusDeformerInstance>(DMCComponent->GetMeshDeformerInstance());
 									if (ensure(MeshDeformerInstance))
